Merge log parser into WorkspaceTool (now renamed to "BuildAgent"). Also add some common functionality into DotNETUtilities.

diff --git a/Engine/Source/Programs/AutomationTool/AutomationUtils/HelpUtils.cs b/Engine/Source/Programs/AutomationTool/AutomationUtils/HelpUtils.cs
--- a/Engine/Source/Programs/AutomationTool/AutomationUtils/HelpUtils.cs
+++ b/Engine/Source/Programs/AutomationTool/AutomationUtils/HelpUtils.cs
@@ -145,11 +145,8 @@ namespace AutomationTool
 				}
 			}
 
-			List<string> Lines = new List<string>();
-			Lines.Add("");
-			Lines.AddRange(HelpUtils.GetHelpText(String.Format("{0} Help:", Command.Name), Description, ParamDict.ToList(), HelpUtils.WindowWidth - 1));
-
-			Log.TraceInformation("{0}", String.Join(Environment.NewLine, Lines));
+			Log.TraceInformation("");
+			HelpUtils.PrintHelp(String.Format("{0} Help:", Command.Name), Description, ParamDict.ToList());
 		}
 
 		/// <summary>
diff --git a/Engine/Source/Programs/AutomationTool/BuildGraph/Graph.cs b/Engine/Source/Programs/AutomationTool/BuildGraph/Graph.cs
--- a/Engine/Source/Programs/AutomationTool/BuildGraph/Graph.cs
+++ b/Engine/Source/Programs/AutomationTool/BuildGraph/Graph.cs
@@ -609,10 +609,7 @@ namespace AutomationTool
 					CommandUtils.LogInformation("");
 					CommandUtils.LogInformation("Options:");
 					CommandUtils.LogInformation("");
-					foreach(string Line in HelpUtils.Tabulate(Parameters, 4, 24, HelpUtils.WindowWidth - 1))
-					{
-						CommandUtils.LogInformation(Line);
-					}
+					HelpUtils.PrintTable(Parameters, 4, 24);
 				}
 			}
 
diff --git a/Engine/Source/Programs/BuildAgent/App.config b/Engine/Source/Programs/BuildAgent/App.config
new file mode 100644
--- /dev/null
+++ b/Engine/Source/Programs/BuildAgent/App.config
@@ -0,0 +1,6 @@
+﻿<?xml version="1.0" encoding="utf-8" ?>
+<configuration>
+    <startup> 
+        <supportedRuntime version="v4.0" sku=".NETFramework,Version=v4.6.2" />
+    </startup>
+</configuration>
\ No newline at end of file
diff --git a/Engine/Source/Programs/BuildAgent/BuildAgent.csproj b/Engine/Source/Programs/BuildAgent/BuildAgent.csproj
new file mode 100644
--- /dev/null
+++ b/Engine/Source/Programs/BuildAgent/BuildAgent.csproj
@@ -0,0 +1,95 @@
+﻿<?xml version="1.0" encoding="utf-8"?>
+<Project ToolsVersion="15.0" xmlns="http://schemas.microsoft.com/developer/msbuild/2003">
+  <Import Project="$(MSBuildExtensionsPath)\$(MSBuildToolsVersion)\Microsoft.Common.props" Condition="Exists('$(MSBuildExtensionsPath)\$(MSBuildToolsVersion)\Microsoft.Common.props')" />
+  <PropertyGroup>
+    <Configuration Condition=" '$(Configuration)' == '' ">Debug</Configuration>
+    <Platform Condition=" '$(Platform)' == '' ">AnyCPU</Platform>
+    <ProjectGuid>{AC5CD7A5-0EA8-41C7-B6FA-79E650E96C29}</ProjectGuid>
+    <OutputType>Exe</OutputType>
+    <RootNamespace>BuildAgent</RootNamespace>
+    <AssemblyName>BuildAgent</AssemblyName>
+    <TargetFrameworkVersion>v4.6.2</TargetFrameworkVersion>
+    <FileAlignment>512</FileAlignment>
+    <AutoGenerateBindingRedirects>true</AutoGenerateBindingRedirects>
+    <Deterministic>true</Deterministic>
+  </PropertyGroup>
+  <PropertyGroup Condition=" '$(Configuration)|$(Platform)' == 'Debug|AnyCPU' ">
+    <PlatformTarget>AnyCPU</PlatformTarget>
+    <DebugSymbols>true</DebugSymbols>
+    <DebugType>full</DebugType>
+    <Optimize>false</Optimize>
+    <OutputPath>bin\Debug\</OutputPath>
+    <DefineConstants>DEBUG;TRACE</DefineConstants>
+    <ErrorReport>prompt</ErrorReport>
+    <WarningLevel>4</WarningLevel>
+  </PropertyGroup>
+  <PropertyGroup Condition=" '$(Configuration)|$(Platform)' == 'Development|AnyCPU' ">
+    <PlatformTarget>AnyCPU</PlatformTarget>
+    <DebugType>pdbonly</DebugType>
+    <Optimize>true</Optimize>
+    <OutputPath>bin\Development\</OutputPath>
+    <DefineConstants>TRACE</DefineConstants>
+    <ErrorReport>prompt</ErrorReport>
+    <WarningLevel>4</WarningLevel>
+  </PropertyGroup>
+  <Import Project="NotForLicensees\EpicInternal.props" Condition="Exists('NotForLicensees\EpicInternal.props')" />
+  <ItemGroup>
+    <Reference Include="System" />
+    <Reference Include="System.Core" />
+    <Reference Include="System.Xml.Linq" />
+    <Reference Include="System.Data.DataSetExtensions" />
+    <Reference Include="Microsoft.CSharp" />
+    <Reference Include="System.Data" />
+    <Reference Include="System.Net.Http" />
+    <Reference Include="System.Xml" />
+  </ItemGroup>
+  <ItemGroup>
+    <Compile Include="..\..\..\Source\Programs\DotNETCommon\MetaData.cs">
+      <Link>Properties\MetaData.cs</Link>
+    </Compile>
+    <Compile Include="Run\RunMode.cs" />
+    <Compile Include="ProgramMode.cs" />
+    <Compile Include="Run\AutoRegisterAttribute.cs" />
+    <Compile Include="Run\LineBuffer.cs" />
+    <Compile Include="Run\LineFilter.cs" />
+    <Compile Include="Run\Listeners\ElectricCommanderListener.cs" />
+    <Compile Include="Run\Listeners\DebugOutputListener.cs" />
+	<Compile Include="Run\Matchers\**\*.cs" />
+    <Compile Include="Program.cs" />
+    <Compile Include="Run\ErrorMatch.cs" />
+    <Compile Include="Run\Interfaces\IErrorListener.cs" />
+    <Compile Include="Run\Interfaces\IErrorMatcher.cs" />
+    <Compile Include="Properties\AssemblyInfo.cs" />
+    <Compile Include="ProgramModeAttribute.cs" />
+    <Compile Include="Workspace\WorkspaceMode.cs" />
+    <Compile Include="Workspace\CleanMode.cs" />
+    <Compile Include="Workspace\ClearMode.cs" />
+    <Compile Include="Workspace\Common\FileContentId.cs" />
+    <Compile Include="Workspace\Common\Repository.cs" />
+    <Compile Include="Workspace\Common\StreamInfo.cs" />
+    <Compile Include="Workspace\Common\TrackedFileInfo.cs" />
+    <Compile Include="Workspace\DumpMode.cs" />
+    <Compile Include="Workspace\PopulateCacheMode.cs" />
+    <Compile Include="Workspace\PurgeCacheMode.cs" />
+    <Compile Include="Workspace\RepairCacheMode.cs" />
+    <Compile Include="Workspace\RevertMode.cs" />
+    <Compile Include="Workspace\SizesMode.cs" />
+    <Compile Include="Workspace\SetupMode.cs" />
+    <Compile Include="Workspace\StatusMode.cs" />
+    <Compile Include="Workspace\SyncMode.cs" />
+    <Compile Include="Workspace\Common\WorkspaceInfo.cs" />
+    <Compile Include="Workspace\Common\WorkspaceTransactionAdd.cs" />
+    <Compile Include="Workspace\Common\WorkspaceTransactionRemove.cs" />
+  </ItemGroup>
+  <ItemGroup>
+    <None Include="App.config" />
+  </ItemGroup>
+  <ItemGroup>
+    <ProjectReference Include="..\..\..\Source\Programs\DotNETCommon\DotNETUtilities\DotNETUtilities.csproj">
+      <Project>{5d7d66e8-8c76-4af9-b3ec-2ef03421d730}</Project>
+      <Name>DotNETUtilities</Name>
+    </ProjectReference>
+  </ItemGroup>
+  <ItemGroup />
+  <Import Project="$(MSBuildToolsPath)\Microsoft.CSharp.targets" />
+</Project>
\ No newline at end of file
diff --git a/Engine/Source/Programs/BuildAgent/Program.cs b/Engine/Source/Programs/BuildAgent/Program.cs
new file mode 100644
--- /dev/null
+++ b/Engine/Source/Programs/BuildAgent/Program.cs
@@ -0,0 +1,165 @@
+// Copyright 1998-2019 Epic Games, Inc. All Rights Reserved.
+
+using System;
+using System.Collections.Generic;
+using System.Diagnostics;
+using System.IO;
+using System.Linq;
+using System.Reflection;
+using System.Text;
+using System.Text.RegularExpressions;
+using System.Threading;
+using System.Threading.Tasks;
+using Tools.DotNETCommon;
+
+namespace BuildAgent
+{
+	class Program
+	{
+		/// <summary>
+		/// Entry point
+		/// </summary>
+		/// <param name="Args">Command-line arguments</param>
+		/// <returns>Exit code</returns>
+		static int Main(string[] Args)
+		{
+			try
+			{
+				int Result = GuardedMain(Args);
+				return Result;
+			}
+			catch (FatalErrorException Ex)
+			{
+				Log.WriteException(Ex, null);
+				return Ex.ExitCode;
+			}
+			catch (Exception Ex)
+			{
+				Log.WriteException(Ex, null);
+				return 1;
+			}
+		}
+
+		/// <summary>
+		/// Actual Main function, without exception guards
+		/// </summary>
+		/// <param name="Args">Command-line arguments</param>
+		/// <returns>Exit code</returns>
+		static int GuardedMain(string[] Args)
+		{
+			// Find the index of the first command
+			int FirstModeIndex = 0;
+			while (FirstModeIndex < Args.Length && Args[FirstModeIndex].StartsWith("-"))
+			{
+				FirstModeIndex++;
+			}
+
+			// Find all the ToolMode types
+			Dictionary<string, Type> ModeNameToType = new Dictionary<string, Type>(StringComparer.OrdinalIgnoreCase);
+			foreach (Type Type in Assembly.GetExecutingAssembly().GetTypes())
+			{
+				ProgramModeAttribute Attribute = Type.GetCustomAttribute<ProgramModeAttribute>();
+				if (Attribute != null)
+				{
+					ModeNameToType.Add(Attribute.Name, Type); 
+				}
+			}
+
+			// Check if there are any commands specified on the command line.
+			if (FirstModeIndex == Args.Length)
+			{
+				Log.TraceInformation("BuildAgent");
+				Log.TraceInformation("");
+				Log.TraceInformation("Utility for managing automated processes on build machines.");
+				Log.TraceInformation("");
+				Log.TraceInformation("Usage:");
+				Log.TraceInformation("    BuildAgent.exe [CommandA] [-OptionA1] [-OptionA2] [CommandB] [-OptionB1]...");
+				Log.TraceInformation("");
+				Log.TraceInformation("Commands:");
+
+				PrintCommands(ModeNameToType);
+				return 0;
+			}
+
+			// Parse the arguments for all commands
+			List<ProgramMode> Modes = new List<ProgramMode>();
+			for (int NextCommandIdx = FirstModeIndex; NextCommandIdx < Args.Length;)
+			{
+				string CommandName = Args[NextCommandIdx++];
+
+				// Get the command type
+				Type CommandType;
+				if (!ModeNameToType.TryGetValue(CommandName, out CommandType))
+				{
+					Log.TraceError("Unknown command '{0}'.", CommandName);
+					Log.TraceInformation("");
+					Log.TraceInformation("Available commands");
+					PrintCommands(ModeNameToType);
+					return 1;
+				}
+
+				// Build an argument list for the command, including all the global arguments as well as arguments until the next command
+				List<string> ArgumentList = new List<string>();
+				for (int Idx = 0; Idx < FirstModeIndex; Idx++)
+				{
+					ArgumentList.Add(Args[Idx]);
+				}
+				while (NextCommandIdx < Args.Length && Args[NextCommandIdx].StartsWith("-"))
+				{
+					ArgumentList.Add(Args[NextCommandIdx++]);
+				}
+				CommandLineArguments CommandArguments = new CommandLineArguments(ArgumentList.ToArray());
+
+				// Create the command instance
+				ProgramMode Command = (ProgramMode)Activator.CreateInstance(CommandType);
+				Modes.Add(Command);
+
+				// If the help flag is specified, print the help info and exit immediately
+				if (CommandArguments.HasOption("-Help"))
+				{
+					HelpUtils.PrintHelp(CommandName, HelpUtils.GetDescription(CommandType), Command.GetParameters());
+					return 1;
+				}
+
+				// Configure the command
+				try
+				{
+					Command.Configure(CommandArguments);
+					CommandArguments.CheckAllArgumentsUsed();
+				}
+				catch (CommandLineArgumentException Ex)
+				{
+					Log.TraceError("{0}: {1}", CommandName, Ex.Message);
+					Log.TraceInformation("");
+					Log.TraceInformation("Arguments for {0}:", CommandName);
+
+					HelpUtils.PrintTable(Command.GetParameters(), 4, 24);
+					return 1;
+				}
+			}
+
+			// Execute all the commands
+			foreach (ProgramMode Mode in Modes)
+			{
+				Mode.Execute();
+			}
+
+			return 0;
+		}
+
+		/// <summary>
+		/// Print a formatted list of all the available commands
+		/// </summary>
+		/// <param name="ModeNameToType">Map from command name to type</param>
+		static void PrintCommands(Dictionary<string, Type> ModeNameToType)
+		{
+			List<KeyValuePair<string, string>> Commands = new List<KeyValuePair<string, string>>();
+			foreach (KeyValuePair<string, Type> Pair in ModeNameToType.OrderBy(x => x.Key))
+			{
+				ProgramModeAttribute Attribute = Pair.Value.GetCustomAttribute<ProgramModeAttribute>();
+				Commands.Add(new KeyValuePair<string, string>(Attribute.Name, Attribute.Description));
+			}
+			HelpUtils.PrintTable(Commands, 4, 20);
+		}
+	}
+}
diff --git a/Engine/Source/Programs/BuildAgent/ProgramMode.cs b/Engine/Source/Programs/BuildAgent/ProgramMode.cs
new file mode 100644
--- /dev/null
+++ b/Engine/Source/Programs/BuildAgent/ProgramMode.cs
@@ -0,0 +1,27 @@
+// Copyright 1998-2019 Epic Games, Inc. All Rights Reserved.
+
+using System;
+using System.Collections.Generic;
+using System.Linq;
+using System.Reflection;
+using System.Text;
+using System.Threading.Tasks;
+using Tools.DotNETCommon;
+
+namespace BuildAgent
+{
+	abstract class ProgramMode
+	{
+		public virtual void Configure(CommandLineArguments Arguments)
+		{
+			Arguments.ApplyTo(this);
+		}
+
+		public virtual List<KeyValuePair<string, string>> GetParameters()
+		{
+			return CommandLineArguments.GetParameters(GetType());
+		}
+
+		public abstract void Execute();
+	}
+}
diff --git a/Engine/Source/Programs/BuildAgent/ProgramModeAttribute.cs b/Engine/Source/Programs/BuildAgent/ProgramModeAttribute.cs
new file mode 100644
--- /dev/null
+++ b/Engine/Source/Programs/BuildAgent/ProgramModeAttribute.cs
@@ -0,0 +1,37 @@
+// Copyright 1998-2019 Epic Games, Inc. All Rights Reserved.
+
+using System;
+using System.Collections.Generic;
+using System.Linq;
+using System.Text;
+using System.Threading.Tasks;
+
+namespace BuildAgent
+{
+	/// <summary>
+	/// Attribute used to specify names of program modes, and help text
+	/// </summary>
+	class ProgramModeAttribute : Attribute
+	{
+		/// <summary>
+		/// Name of the program mode, which needs to be specified on the command line to execute it
+		/// </summary>
+		public string Name;
+
+		/// <summary>
+		/// Short description for the mode. Will be displayed in the help text.
+		/// </summary>
+		public string Description;
+
+		/// <summary>
+		/// Constructor
+		/// </summary>
+		/// <param name="Name">Name of the mode</param>
+		/// <param name="Description">Short description for display in the help text</param>
+		public ProgramModeAttribute(string Name, string Description)
+		{
+			this.Name = Name;
+			this.Description = Description;
+		}
+	}
+}
diff --git a/Engine/Source/Programs/BuildAgent/Properties/AssemblyInfo.cs b/Engine/Source/Programs/BuildAgent/Properties/AssemblyInfo.cs
new file mode 100644
--- /dev/null
+++ b/Engine/Source/Programs/BuildAgent/Properties/AssemblyInfo.cs
@@ -0,0 +1,10 @@
+// Copyright 1998-2019 Epic Games, Inc. All Rights Reserved.
+
+using System.Reflection;
+
+// General Information about an assembly is controlled through the following 
+// set of attributes. Change these attribute values to modify the information
+// associated with an assembly.
+[assembly: AssemblyTitle("BuildAgent")]
+[assembly: AssemblyDescription("")]
+[assembly: AssemblyConfiguration("")]
diff --git a/Engine/Source/Programs/BuildAgent/Run/AutoRegisterAttribute.cs b/Engine/Source/Programs/BuildAgent/Run/AutoRegisterAttribute.cs
new file mode 100644
--- /dev/null
+++ b/Engine/Source/Programs/BuildAgent/Run/AutoRegisterAttribute.cs
@@ -0,0 +1,18 @@
+// Copyright 1998-2019 Epic Games, Inc. All Rights Reserved.
+
+using System;
+using System.Collections.Generic;
+using System.Linq;
+using System.Text;
+using System.Threading.Tasks;
+
+namespace BuildAgent.Run
+{
+	/// <summary>
+	/// Attribute used to indicate that an ErrorMatcher class should be created and registered automatically at startup
+	/// </summary>
+	[AttributeUsage(AttributeTargets.Class)]
+	class AutoRegisterAttribute : Attribute
+	{
+	}
+}
diff --git a/Engine/Source/Programs/BuildAgent/Run/ErrorMatch.cs b/Engine/Source/Programs/BuildAgent/Run/ErrorMatch.cs
new file mode 100644
--- /dev/null
+++ b/Engine/Source/Programs/BuildAgent/Run/ErrorMatch.cs
@@ -0,0 +1,62 @@
+// Copyright 1998-2019 Epic Games, Inc. All Rights Reserved.
+
+using System;
+using System.Collections.Generic;
+using System.Linq;
+using System.Text;
+using System.Threading.Tasks;
+
+namespace BuildAgent.Run
+{
+	enum ErrorSeverity
+	{
+		Error,
+		Warning,
+		Silent,
+	}
+
+	enum ErrorPriority
+	{
+		Lowest,
+		Low,
+		BelowNormal,
+		Normal,
+		AboveNormal,
+		High,
+		Highest,
+	}
+
+	class ErrorMatch
+	{
+		public ErrorSeverity Severity;
+		public ErrorPriority Priority;
+		public string Type;
+		public int MinLineNumber;
+		public int MaxLineNumber;
+		public List<string> Lines = new List<string>();
+		public Dictionary<string, string> Properties = new Dictionary<string, string>(StringComparer.Ordinal);
+
+		public ErrorMatch(ErrorSeverity Severity, ErrorPriority Priority, string Type, int MinLineNumber, int MaxLineNumber)
+		{
+			this.Severity = Severity;
+			this.Priority = Priority;
+			this.Type = Type;
+			this.MinLineNumber = MinLineNumber;
+			this.MaxLineNumber = MaxLineNumber;
+		}
+
+		public ErrorMatch(ErrorSeverity Severity, ErrorPriority Priority, string Type, ReadOnlyLineBuffer Input)
+			: this(Severity, Priority, Type, Input, 0, 0)
+		{
+		}
+
+		public ErrorMatch(ErrorSeverity Severity, ErrorPriority Priority, string Type, ReadOnlyLineBuffer Input, int MinOffset, int MaxOffset)
+			: this(Severity, Priority, Type, Input.CurrentLineNumber + MinOffset, Input.CurrentLineNumber + MaxOffset)
+		{
+			for (int Offset = MinOffset; Offset <= MaxOffset; Offset++)
+			{
+				Lines.Add(Input[Offset]);
+			}
+		}
+	}
+}
diff --git a/Engine/Source/Programs/BuildAgent/Run/Interfaces/IErrorListener.cs b/Engine/Source/Programs/BuildAgent/Run/Interfaces/IErrorListener.cs
new file mode 100644
--- /dev/null
+++ b/Engine/Source/Programs/BuildAgent/Run/Interfaces/IErrorListener.cs
@@ -0,0 +1,22 @@
+// Copyright 1998-2019 Epic Games, Inc. All Rights Reserved.
+
+using System;
+using System.Collections.Generic;
+using System.Linq;
+using System.Text;
+using System.Threading.Tasks;
+
+namespace BuildAgent.Run.Interfaces
+{
+	/// <summary>
+	/// Interface for classes which listen for error matches
+	/// </summary>
+	interface IErrorListener : IDisposable
+	{
+		/// <summary>
+		/// Called when an error is matched
+		/// </summary>
+		/// <param name="Error">The matched error</param>
+		void OnErrorMatch(ErrorMatch Error);
+	}
+}
diff --git a/Engine/Source/Programs/BuildAgent/Run/Interfaces/IErrorMatcher.cs b/Engine/Source/Programs/BuildAgent/Run/Interfaces/IErrorMatcher.cs
new file mode 100644
--- /dev/null
+++ b/Engine/Source/Programs/BuildAgent/Run/Interfaces/IErrorMatcher.cs
@@ -0,0 +1,23 @@
+// Copyright 1998-2019 Epic Games, Inc. All Rights Reserved.
+
+using System;
+using System.Collections.Generic;
+using System.Linq;
+using System.Text;
+using System.Threading.Tasks;
+
+namespace BuildAgent.Run.Interfaces
+{
+	/// <summary>
+	/// Interface for a class which matches error strings
+	/// </summary>
+	interface IErrorMatcher
+	{
+		/// <summary>
+		/// Attempt to match an error from the given input buffer
+		/// </summary>
+		/// <param name="Input">The input buffer</param>
+		/// <returns>Information about the error that was matched, or null if an error was not matched</returns>
+		ErrorMatch Match(ReadOnlyLineBuffer Input);
+	}
+}
diff --git a/Engine/Source/Programs/BuildAgent/Run/LineBuffer.cs b/Engine/Source/Programs/BuildAgent/Run/LineBuffer.cs
new file mode 100644
--- /dev/null
+++ b/Engine/Source/Programs/BuildAgent/Run/LineBuffer.cs
@@ -0,0 +1,213 @@
+// Copyright 1998-2019 Epic Games, Inc. All Rights Reserved.
+
+using System;
+using System.Collections.Generic;
+using System.Linq;
+using System.Text;
+using System.Text.RegularExpressions;
+using System.Threading.Tasks;
+
+namespace BuildAgent.Run
+{
+	class LineBuffer
+	{
+		int LineNumber;
+
+		string[] History;
+		int HistoryIdx;
+		int HistoryCount;
+
+		List<string> NextLines;
+		Func<string> ReadLine;
+
+		public LineBuffer(Func<string> ReadLine, int HistorySize)
+		{
+			LineNumber = 1;
+			History = new string[HistorySize];
+			NextLines = new List<string>();
+
+			this.ReadLine = ReadLine;
+		}
+
+		public int CurrentLineNumber
+		{
+			get { return LineNumber; }
+		}
+
+		public void Advance(int Count)
+		{
+			for(int Idx = 0; Idx < Count; Idx++)
+			{
+				MoveNext();
+			}
+		}
+
+		public void MoveNext()
+		{
+			if(NextLines.Count == 0)
+			{
+				NextLines.Add(ReadLine());
+			}
+
+			HistoryIdx++;
+			if (HistoryIdx >= History.Length)
+			{
+				HistoryIdx = 0;
+			}
+			if (HistoryCount < History.Length)
+			{
+				HistoryCount++;
+			}
+
+			History[HistoryIdx] = NextLines[0];
+			NextLines.RemoveAt(0);
+
+			LineNumber++;
+		}
+
+		public bool TryGetLine(int Offset, out string NextLine)
+		{
+			NextLine = this[Offset];
+			return NextLine != null;
+		}
+
+		public string this[int Offset]
+		{
+			get
+			{
+				if (Offset >= 0)
+				{
+					// Add new lines to the buffer
+					while (Offset >= NextLines.Count)
+					{
+						NextLines.Add(ReadLine());
+					}
+					return NextLines[Offset];
+				}
+				else if (Offset >= -HistoryCount)
+				{
+					// Retrieve a line from the history buffer
+					int Idx = HistoryIdx + 1 + Offset;
+					if (Idx < 0)
+					{
+						Idx += History.Length;
+					}
+					return History[Idx];
+				}
+				else
+				{
+					// Invalid index
+					throw new ArgumentException(String.Format("Invalid line buffer offset ({0})", Offset));
+				}
+			}
+		}
+
+		public override string ToString()
+		{
+			return String.Join("\n", NextLines);
+		}
+	}
+
+	class ReadOnlyLineBuffer
+	{
+		LineBuffer Inner;
+
+		public ReadOnlyLineBuffer(LineBuffer Inner)
+		{
+			this.Inner = Inner;
+		}
+
+		public string CurrentLine
+		{
+			get { return Inner[0]; }
+		}
+
+		public int CurrentLineNumber
+		{
+			get { return Inner.CurrentLineNumber; }
+		}
+
+		public string this[int Offset]
+		{
+			get { return Inner[Offset]; }
+		}
+
+		public bool IsMatch(string Pattern)
+		{
+			return IsMatch(0, Pattern);
+		}
+
+		public bool IsMatch(int Offset, string Pattern)
+		{
+			string Line;
+			return TryGetLine(Offset, out Line) && Regex.IsMatch(Line, Pattern);
+		}
+
+		public bool TryMatch(string Pattern, out Match OutMatch)
+		{
+			return TryMatch(0, Pattern, out OutMatch);
+		}
+
+		public bool TryMatch(int Offset, string Pattern, out Match OutMatch)
+		{
+			string Line;
+			if(!TryGetLine(Offset, out Line))
+			{
+				OutMatch = null;
+				return false;
+			}
+
+			Match Match = Regex.Match(Line, Pattern);
+			if(!Match.Success)
+			{
+				OutMatch = null;
+				return false;
+			}
+
+			OutMatch = Match;
+			return true;
+		}
+
+		public int MatchBackwards(int StartOffset, string Pattern)
+		{
+			int Offset = StartOffset;
+			while(IsMatch(Offset - 1, Pattern))
+			{
+				Offset--;
+			}
+			return Offset;
+		}
+
+		public int MatchForwards(int Offset, string Pattern)
+		{
+			while(IsMatch(Offset + 1, Pattern))
+			{
+				Offset++;
+			}
+			return Offset;
+		}
+
+		public int MatchForwardsUntil(int Offset, string Pattern)
+		{
+			string NextLine;
+			for (int NextOffset = Offset + 1; TryGetLine(NextOffset, out NextLine); NextOffset++)
+			{
+				if(Regex.IsMatch(NextLine, Pattern))
+				{
+					return NextOffset;
+				}
+			}
+			return Offset;
+		}
+
+		public bool TryGetLine(int Offset, out string NextLine)
+		{
+			return Inner.TryGetLine(Offset, out NextLine);
+		}
+
+		public override string ToString()
+		{
+			return CurrentLine;
+		}
+	}
+}
diff --git a/Engine/Source/Programs/BuildAgent/Run/LineFilter.cs b/Engine/Source/Programs/BuildAgent/Run/LineFilter.cs
new file mode 100644
--- /dev/null
+++ b/Engine/Source/Programs/BuildAgent/Run/LineFilter.cs
@@ -0,0 +1,56 @@
+// Copyright 1998-2019 Epic Games, Inc. All Rights Reserved.
+
+using System;
+using System.Collections.Generic;
+using System.Linq;
+using System.Text;
+using System.Text.RegularExpressions;
+using System.Threading.Tasks;
+
+namespace BuildAgent.Run
+{
+	class LineFilter
+	{
+		Func<string> InnerReadLine;
+		int OutputEnabled;
+
+		public LineFilter(Func<string> ReadLine)
+		{
+			this.InnerReadLine = ReadLine;
+		}
+
+		public string ReadLine()
+		{
+			string Line = InnerReadLine();
+			if (Line != null)
+			{
+				// remove timestamps added by UAT
+				Line = Regex.Replace(Line, @"^\[[0-9:.]+\] ", "");
+
+				// remove the mono prefix for recursive UAT calls on Mac
+				Line = Regex.Replace(Line, @"^mono: ([A-Za-z0-9-]+: )", "$1");
+
+				// remove the AutomationTool prefix for recursive UAT calls anywhere
+				Line = Regex.Replace(Line, @"^AutomationTool: ([A-Za-z0-9-]+: )", "$1");
+
+				// if it was the editor, also strip any additional timestamp
+				Line = Regex.Replace(Line, @"^\s*\[[0-9.:]+-[0-9.:]+\]\[\s*\d+\]", "");
+
+				// look for a special marker to disable output
+				if (Regex.IsMatch(Line, "<-- Suspend Log Parsing -->", RegexOptions.IgnoreCase))
+				{
+					OutputEnabled--;
+				}
+				else if (Regex.IsMatch(Line, "<-- Resume Log Parsing -->", RegexOptions.IgnoreCase))
+				{
+					OutputEnabled++;
+				}
+				else if (OutputEnabled < 0)
+				{
+					Line = "";
+				}
+			}
+			return Line;
+		}
+	}
+}
diff --git a/Engine/Source/Programs/BuildAgent/Run/Listeners/DebugOutputListener.cs b/Engine/Source/Programs/BuildAgent/Run/Listeners/DebugOutputListener.cs
new file mode 100644
--- /dev/null
+++ b/Engine/Source/Programs/BuildAgent/Run/Listeners/DebugOutputListener.cs
@@ -0,0 +1,47 @@
+// Copyright 1998-2019 Epic Games, Inc. All Rights Reserved.
+
+using BuildAgent.Run;
+using BuildAgent.Run.Interfaces;
+using System;
+using System.Collections.Generic;
+using System.Linq;
+using System.Text;
+using System.Threading.Tasks;
+
+namespace BuildAgent.Listeners
+{
+	/// <summary>
+	/// Debug error listener. Just prints the error output to the log.
+	/// </summary>
+	class DebugOutputListener : IErrorListener
+	{
+		/// <summary>
+		/// Dispose of this object
+		/// </summary>
+		public void Dispose()
+		{
+		}
+
+		/// <summary>
+		/// Called when an error is matched
+		/// </summary>
+		/// <param name="Error">The matched error</param>
+		public void OnErrorMatch(ErrorMatch Error)
+		{
+			Console.WriteLine("Error of type '{0}' at line {1}:", Error.Type, Error.MinLineNumber);
+			foreach(string Line in Error.Lines)
+			{
+				Console.WriteLine("Text: |{0}|", Line);
+			}
+			if (Error.Properties.Count > 0)
+			{
+				Console.WriteLine("Properties:");
+				foreach (KeyValuePair<string, string> Pair in Error.Properties)
+				{
+					Console.WriteLine("  {0} = {1}", Pair.Key, Pair.Value);
+				}
+			}
+			Console.WriteLine();
+		}
+	}
+}
diff --git a/Engine/Source/Programs/BuildAgent/Run/Listeners/ElectricCommanderListener.cs b/Engine/Source/Programs/BuildAgent/Run/Listeners/ElectricCommanderListener.cs
new file mode 100644
--- /dev/null
+++ b/Engine/Source/Programs/BuildAgent/Run/Listeners/ElectricCommanderListener.cs
@@ -0,0 +1,149 @@
+// Copyright 1998-2019 Epic Games, Inc. All Rights Reserved.
+
+using BuildAgent.Run;
+using BuildAgent.Run.Interfaces;
+using System;
+using System.Collections.Generic;
+using System.Diagnostics;
+using System.IO;
+using System.Linq;
+using System.Text;
+using System.Threading.Tasks;
+using System.Xml;
+using Tools.DotNETCommon;
+
+namespace BuildAgent.Listeners
+{
+	/// <summary>
+	/// Writes parsed errors to a log file usable by Electric Commander
+	/// </summary>
+	class ElectricCommanderListener : IErrorListener
+	{
+		/// <summary>
+		/// Path to the output file
+		/// </summary>
+		FileReference OutputFile;
+
+		/// <summary>
+		/// List of the matched errors
+		/// </summary>
+		List<ErrorMatch> Errors = new List<ErrorMatch>();
+
+		/// <summary>
+		/// Whether we should set a property indicating that the diagnostics file exists
+		/// </summary>
+		bool bUpdateProperty = true;
+
+		/// <summary>
+		/// Process which is created to update properties in the background
+		/// </summary>
+		Process UpdatePropertyProcess;
+
+		/// <summary>
+		/// Constructs an EC listener
+		/// </summary>
+		/// <param name="OutputFile">Path to the output file</param>
+		public ElectricCommanderListener()
+		{
+			string EnvVar = Environment.GetEnvironmentVariable("COMMANDER_JOBSTEP");
+			if (String.IsNullOrEmpty(EnvVar))
+			{
+				OutputFile = new FileReference("BuildAgent.xml");
+				bUpdateProperty = false;
+			}
+			else
+			{
+				OutputFile = new FileReference(String.Format("BuildAgent-{0}.xml", EnvVar));
+				bUpdateProperty = true;
+			}
+		}
+
+		/// <summary>
+		/// Wait for the child process to finish
+		/// </summary>
+		public void Dispose()
+		{
+			if(UpdatePropertyProcess != null)
+			{
+				if (!UpdatePropertyProcess.WaitForExit(0))
+				{
+					Log.TraceInformation("Waiting for ectool to terminate...");
+					UpdatePropertyProcess.WaitForExit();
+				}
+				if (UpdatePropertyProcess.ExitCode != 0)
+				{
+					Log.TraceWarning("ectool terminated with exit code {0}", UpdatePropertyProcess.ExitCode);
+				}
+				UpdatePropertyProcess.Dispose();
+				UpdatePropertyProcess = null;
+			}
+		}
+
+		/// <summary>
+		/// Called when an error is matched
+		/// </summary>
+		/// <param name="Error">The matched error</param>
+		public void OnErrorMatch(ErrorMatch Error)
+		{
+			Errors.Add(Error);
+
+			Write();
+
+			if(bUpdateProperty)
+			{
+				UpdatePropertyProcess = Process.Start("ectool", String.Format("setProperty /myJobStep/diagFile {0}", OutputFile.GetFileName()));
+				bUpdateProperty = false;
+			}
+		}
+
+		/// <summary>
+		/// Write the list of errors to disk in XML format
+		/// </summary>
+		public void Write()
+		{
+			XmlWriterSettings Settings = new XmlWriterSettings();
+			Settings.NewLineChars = Environment.NewLine;
+			Settings.Indent = true;
+			Settings.IndentChars = "\t";
+			using (FileStream Stream = File.Open(OutputFile.FullName, FileMode.Create, FileAccess.Write, FileShare.Read | FileShare.Delete))
+			{
+				using (XmlWriter Writer = XmlWriter.Create(Stream, Settings))
+				{
+					Writer.WriteStartElement("diagnostics");
+					foreach (ErrorMatch Error in Errors)
+					{
+						if (Error.Severity != ErrorSeverity.Silent)
+						{
+							Writer.WriteStartElement("diagnostic");
+
+							Writer.WriteElementString("matcher", Error.Type);
+							Writer.WriteElementString("name", "");
+
+							if(Error.Severity == ErrorSeverity.Error)
+							{
+								Writer.WriteElementString("type", "error");
+							}
+							else
+							{
+								Writer.WriteElementString("type", "warning");
+							}
+
+							Writer.WriteElementString("firstLine", Error.MinLineNumber.ToString());
+							Writer.WriteElementString("numLines", (Error.MaxLineNumber + 1 - Error.MinLineNumber).ToString());
+
+							StringBuilder Message = new StringBuilder();
+							foreach(string Line in Error.Lines)
+							{
+								Message.AppendLine(Line);
+							}
+							Writer.WriteElementString("message", Message.ToString());
+
+							Writer.WriteEndElement();
+						}
+					}
+					Writer.WriteEndElement();
+				}
+			}
+		}
+	}
+}
diff --git a/Engine/Source/Programs/BuildAgent/Run/Matchers/CompileErrorMatcher.cs b/Engine/Source/Programs/BuildAgent/Run/Matchers/CompileErrorMatcher.cs
new file mode 100644
--- /dev/null
+++ b/Engine/Source/Programs/BuildAgent/Run/Matchers/CompileErrorMatcher.cs
@@ -0,0 +1,90 @@
+// Copyright 1998-2019 Epic Games, Inc. All Rights Reserved.
+
+using BuildAgent.Run;
+using BuildAgent.Run.Interfaces;
+using System;
+using System.Collections.Generic;
+using System.Linq;
+using System.Text;
+using System.Text.RegularExpressions;
+using System.Threading.Tasks;
+
+namespace BuildAgent.Matchers
+{
+	[AutoRegister]
+	class CompileErrorMatcher : IErrorMatcher
+	{
+		const string ClangFileLinePattern =
+			@"("                        +
+			    @"(?:[a-zA-Z]:)?"		+	// optional drive letter
+			    @"[^:]+"				+	// any non-colon character
+			@")"						+
+			@"(?:"						+
+				@"\s*:[\s\d:,]+:"		+	// clang-style   :123:456:
+				@"|"					+
+				@"\([\s\d:,]+\):"		+ 	// msvc-style    (123,456):
+			@")"						;
+
+		public ErrorMatch Match(ReadOnlyLineBuffer Input)
+		{
+			// Do the match in two phases so we can early out if the strings "error" or "warning" are not present. The patterns before these strings can
+			// produce many false positives, making them very slow to execute.
+			if (Input.IsMatch("error"))
+			{
+				if (Input.IsMatch(@"([^(]+)(\([\d,]+\))? ?: (fatal )?error [a-zA-Z]+[\d]+"))
+				{
+					return ParseVisualCppMatch(Input, ErrorSeverity.Error);
+				}
+				if (Input.IsMatch(@"\s*" + ClangFileLinePattern + @"\s*error\s*:"))
+				{
+					return ParseClangMatch(Input, ErrorSeverity.Error);
+				}
+			}
+			if (Input.IsMatch("warning"))
+			{
+				if (Input.IsMatch(@"([^(]+)(\([\d,]+\))? ?: warning[ :]"))
+				{
+					return ParseVisualCppMatch(Input, ErrorSeverity.Warning);
+				}
+				if (Input.IsMatch(@"\s*" + ClangFileLinePattern + @"\s*warning\s*:"))
+				{
+					return ParseClangMatch(Input, ErrorSeverity.Warning);
+				}
+			}
+			return null;
+		}
+
+		static ErrorMatch ParseVisualCppMatch(ReadOnlyLineBuffer Input, ErrorSeverity Severity)
+		{
+			string Pattern = String.Format("^{0} |: note:", ExtractIndent(Input[0]));
+
+			int EndIdx = 0;
+			while (Input.IsMatch(EndIdx + 1, Pattern))
+			{
+				EndIdx++;
+			}
+
+			return new ErrorMatch(Severity, ErrorPriority.High, "Compile", Input, 0, EndIdx);
+		}
+
+		static ErrorMatch ParseClangMatch(ReadOnlyLineBuffer Input, ErrorSeverity Severity)
+		{
+			string Indent = ExtractIndent(Input[0]);
+
+			int MinOffset = Input.MatchBackwards(0, String.Format(@"^(?:{0}).*(?:In (member )?function|In file included from)", Indent));
+			int MaxOffset = Input.MatchForwards(0, String.Format(@"^({0} |{0}{1}\s*note:| *\$)", Indent, ClangFileLinePattern));
+
+			return new ErrorMatch(Severity, ErrorPriority.High, "Compile", Input, MinOffset, MaxOffset);
+		}
+
+		static string ExtractIndent(string Line)
+		{
+			int Length = 0;
+			while(Length < Line.Length && Line[Length] == ' ')
+			{
+				Length++;
+			}
+			return new string(' ', Length);
+		}
+	}
+}
diff --git a/Engine/Source/Programs/BuildAgent/Run/Matchers/CrashErrorMatcher.cs b/Engine/Source/Programs/BuildAgent/Run/Matchers/CrashErrorMatcher.cs
new file mode 100644
--- /dev/null
+++ b/Engine/Source/Programs/BuildAgent/Run/Matchers/CrashErrorMatcher.cs
@@ -0,0 +1,41 @@
+// Copyright 1998-2019 Epic Games, Inc. All Rights Reserved.
+
+using BuildAgent.Run;
+using BuildAgent.Run.Interfaces;
+using System;
+using System.Collections.Generic;
+using System.Linq;
+using System.Text;
+using System.Text.RegularExpressions;
+using System.Threading.Tasks;
+
+namespace BuildAgent.Matchers
+{
+	[AutoRegister]
+	class CrashErrorMatcher : IErrorMatcher
+	{
+		public ErrorMatch Match(ReadOnlyLineBuffer Input)
+		{
+			if(Regex.IsMatch(Input[0], "begin: stack for UAT"))
+			{
+				for (int Idx = 1; Idx < 100; Idx++)
+				{
+					if(Regex.IsMatch(Input[Idx], "end: stack for UAT"))
+					{
+						return new ErrorMatch(ErrorSeverity.Error, ErrorPriority.BelowNormal, "crash", Input, 0, Idx);
+					}
+				}
+			}
+			if(Input.IsMatch("AutomationTool: Stack:"))
+			{
+				int EndOffset = Input.MatchForwards(0, "AutomationTool: Stack:");
+				return new ErrorMatch(ErrorSeverity.Error, ErrorPriority.Low, "Crash", Input, 0, EndOffset);
+			}
+			if (Input.IsMatch("ExitCode=(3|139|255)"))
+			{
+				return new ErrorMatch(ErrorSeverity.Error, ErrorPriority.Low, "Crash", Input);
+			}
+			return null;
+		}
+	}
+}
diff --git a/Engine/Source/Programs/BuildAgent/Run/Matchers/ExceptionErrorMatcher.cs b/Engine/Source/Programs/BuildAgent/Run/Matchers/ExceptionErrorMatcher.cs
new file mode 100644
--- /dev/null
+++ b/Engine/Source/Programs/BuildAgent/Run/Matchers/ExceptionErrorMatcher.cs
@@ -0,0 +1,26 @@
+// Copyright 1998-2019 Epic Games, Inc. All Rights Reserved.
+
+using BuildAgent.Run;
+using BuildAgent.Run.Interfaces;
+using System;
+using System.Collections.Generic;
+using System.Linq;
+using System.Text;
+using System.Threading.Tasks;
+
+namespace BuildAgent.Matchers
+{
+	[AutoRegister]
+	class ExceptionErrorMatcher : IErrorMatcher
+	{
+		public ErrorMatch Match(ReadOnlyLineBuffer Input)
+		{
+			if (Input.IsMatch(@"^\s*Unhandled Exception: "))
+			{
+				int MaxOffset = Input.MatchForwards(0, @"^\s*at ");
+				return new ErrorMatch(ErrorSeverity.Error, ErrorPriority.Low, "Exception", Input, 0, MaxOffset);
+			}
+			return null;
+		}
+	}
+}
diff --git a/Engine/Source/Programs/BuildAgent/Run/Matchers/GenericErrorMatcher.cs b/Engine/Source/Programs/BuildAgent/Run/Matchers/GenericErrorMatcher.cs
new file mode 100644
--- /dev/null
+++ b/Engine/Source/Programs/BuildAgent/Run/Matchers/GenericErrorMatcher.cs
@@ -0,0 +1,50 @@
+// Copyright 1998-2019 Epic Games, Inc. All Rights Reserved.
+
+using BuildAgent.Run;
+using BuildAgent.Run.Interfaces;
+using System;
+using System.Collections.Generic;
+using System.Linq;
+using System.Text;
+using System.Text.RegularExpressions;
+using System.Threading.Tasks;
+
+namespace BuildAgent
+{
+	[AutoRegister]
+	class GenericErrorMatcher : IErrorMatcher
+	{
+		public ErrorMatch Match(ReadOnlyLineBuffer Input)
+		{
+			if (Input.IsMatch(@"^\s*FATAL:"))
+			{
+				return new ErrorMatch(ErrorSeverity.Error, ErrorPriority.Low, "Generic", Input);
+			}
+			if (Input.IsMatch(@"(?<!\w)(ERROR|[Ee]rror)( (\([^)]+\)|\[[^\]]+\]))?: "))
+			{
+				int MaxOffset = Input.MatchForwards(0, String.Format(@"^({0} | *$)", ExtractIndent(Input[0])));
+				return new ErrorMatch(ErrorSeverity.Error, ErrorPriority.Lowest, "Generic", Input, 0, MaxOffset);
+			}
+			if (Input.IsMatch(@"(?<!\w)(WARNING|[Ww]arning)( (\([^)]+\)|\[[^\]]+\]))?: "))
+			{
+				int MaxOffset = Input.MatchForwards(0, String.Format(@"^({0} | *$)", ExtractIndent(Input[0])));
+				return new ErrorMatch(ErrorSeverity.Warning, ErrorPriority.Lowest, "Generic", Input, 0, MaxOffset);
+			}
+			if (Input.IsMatch(@"[Ee]rror [A-Z]\d+\s:"))
+			{
+				return new ErrorMatch(ErrorSeverity.Error, ErrorPriority.Lowest, "Generic", Input);
+			}
+			return null;
+		}
+
+		static string ExtractIndent(string Line)
+		{
+			int Length = 0;
+			while (Length < Line.Length && Line[Length] == ' ')
+			{
+				Length++;
+			}
+			return new string(' ', Length);
+		}
+	}
+}
diff --git a/Engine/Source/Programs/BuildAgent/Run/Matchers/LinkErrorMatcher.cs b/Engine/Source/Programs/BuildAgent/Run/Matchers/LinkErrorMatcher.cs
new file mode 100644
--- /dev/null
+++ b/Engine/Source/Programs/BuildAgent/Run/Matchers/LinkErrorMatcher.cs
@@ -0,0 +1,35 @@
+// Copyright 1998-2019 Epic Games, Inc. All Rights Reserved.
+
+using BuildAgent.Run;
+using BuildAgent.Run.Interfaces;
+using System;
+using System.Collections.Generic;
+using System.Linq;
+using System.Text;
+using System.Threading.Tasks;
+
+namespace BuildAgent.Matchers
+{
+	[AutoRegister]
+	class LinkErrorMatcher : IErrorMatcher
+	{
+		public ErrorMatch Match(ReadOnlyLineBuffer Input)
+		{
+			if(Input.IsMatch(@"error: linker command failed with exit code "))
+			{
+				int MinIdx = Input.MatchBackwards(0, ": In function |: undefined reference to ");
+				return new ErrorMatch(ErrorSeverity.Error, ErrorPriority.Normal, "Link", Input, MinIdx, 0);
+			}
+			if (Input.IsMatch("Undefined symbols for architecture"))
+			{
+				int MaxOffset = Input.MatchForwardsUntil(0, "ld: symbol");
+				return new ErrorMatch(ErrorSeverity.Error, ErrorPriority.Normal, "Link", Input, 0, MaxOffset);
+			}
+			if (Input.IsMatch(@"LINK : fatal error"))
+			{
+				return new ErrorMatch(ErrorSeverity.Error, ErrorPriority.Normal, "Link", Input);
+			}
+			return null;
+		}
+	}
+}
diff --git a/Engine/Source/Programs/BuildAgent/Run/Matchers/LogChannelErrorMatcher.cs b/Engine/Source/Programs/BuildAgent/Run/Matchers/LogChannelErrorMatcher.cs
new file mode 100644
--- /dev/null
+++ b/Engine/Source/Programs/BuildAgent/Run/Matchers/LogChannelErrorMatcher.cs
@@ -0,0 +1,45 @@
+// Copyright 1998-2019 Epic Games, Inc. All Rights Reserved.
+
+using BuildAgent.Run;
+using BuildAgent.Run.Interfaces;
+using System;
+using System.Collections.Generic;
+using System.Linq;
+using System.Text;
+using System.Text.RegularExpressions;
+using System.Threading.Tasks;
+
+namespace BuildAgent.Matchers
+{
+	[AutoRegister]
+	class LogChannelErrorMatcher : IErrorMatcher
+	{
+		public ErrorMatch Match(ReadOnlyLineBuffer Input)
+		{
+			Match Match;
+			if(Input.TryMatch(@"^(\s*)([a-zA-Z_][a-zA-Z0-9_]*):\s*(Error|Warning|Display): ", out Match))
+			{
+				int MaxOffset = Input.MatchForwards(0, String.Format(@"^({0} | *$)", Match.Groups[1].Value));
+
+				ErrorSeverity Severity;
+				switch(Match.Groups[3].Value)
+				{
+					case "Error":
+						Severity = ErrorSeverity.Error;
+						break;
+					case "Warning":
+						Severity = ErrorSeverity.Warning;
+						break;
+					default:
+						Severity = ErrorSeverity.Silent;
+						break;
+				}
+
+				ErrorMatch Error = new ErrorMatch(Severity, ErrorPriority.Low, "Log", Input, 0, MaxOffset);
+				Error.Properties["Channel"] = Match.Groups[2].Value;
+				return Error;
+			}
+			return null;
+		}
+	}
+}
diff --git a/Engine/Source/Programs/BuildAgent/Run/Matchers/SanitizerErrorMatcher.cs b/Engine/Source/Programs/BuildAgent/Run/Matchers/SanitizerErrorMatcher.cs
new file mode 100644
--- /dev/null
+++ b/Engine/Source/Programs/BuildAgent/Run/Matchers/SanitizerErrorMatcher.cs
@@ -0,0 +1,28 @@
+// Copyright 1998-2019 Epic Games, Inc. All Rights Reserved.
+
+using BuildAgent.Run;
+using BuildAgent.Run.Interfaces;
+using System;
+using System.Collections.Generic;
+using System.Linq;
+using System.Text;
+using System.Text.RegularExpressions;
+using System.Threading.Tasks;
+
+namespace BuildAgent.Matchers
+{
+	[AutoRegister]
+	class SanitizerErrorMatcher : IErrorMatcher
+	{
+		public ErrorMatch Match(ReadOnlyLineBuffer Input)
+		{
+			Match Match;
+			if (Input.TryMatch(@"^(\s*)WARNING: ThreadSanitizer:", out Match))
+			{
+				int EndIdx = Input.MatchForwards(0, String.Format(@"^([ ]*|{0}  .*|{0}SUMMARY:.*)\$", Match.Groups[1].Value));
+				return new ErrorMatch(ErrorSeverity.Warning, ErrorPriority.Normal, "ThreadSanitizer", Input, 0, EndIdx);
+			}
+			return null;
+		}
+	}
+}
diff --git a/Engine/Source/Programs/BuildAgent/Run/Matchers/XoreaxErrorMatcher.cs b/Engine/Source/Programs/BuildAgent/Run/Matchers/XoreaxErrorMatcher.cs
new file mode 100644
--- /dev/null
+++ b/Engine/Source/Programs/BuildAgent/Run/Matchers/XoreaxErrorMatcher.cs
@@ -0,0 +1,26 @@
+// Copyright 1998-2019 Epic Games, Inc. All Rights Reserved.
+
+using BuildAgent.Run;
+using BuildAgent.Run.Interfaces;
+using System;
+using System.Collections.Generic;
+using System.Linq;
+using System.Text;
+using System.Threading.Tasks;
+
+namespace BuildAgent.Matchers
+{
+	[AutoRegister]
+	class XoreaxErrorMatcher : IErrorMatcher
+	{
+		public ErrorMatch Match(ReadOnlyLineBuffer Input)
+		{
+			if(Input.IsMatch(@"\s*----+Build System Warning----+"))
+			{
+				int MaxOffset = Input.MatchForwards(0, @"^\s*[A-Za-z]");
+				return new ErrorMatch(ErrorSeverity.Warning, ErrorPriority.Normal, "Xoreax", Input, 0, MaxOffset);
+			}
+			return null;
+		}
+	}
+}
diff --git a/Engine/Source/Programs/BuildAgent/Run/RunMode.cs b/Engine/Source/Programs/BuildAgent/Run/RunMode.cs
new file mode 100644
--- /dev/null
+++ b/Engine/Source/Programs/BuildAgent/Run/RunMode.cs
@@ -0,0 +1,243 @@
+// Copyright 1998-2019 Epic Games, Inc. All Rights Reserved.
+
+using BuildAgent.Listeners;
+using BuildAgent.Run.Interfaces;
+using System;
+using System.Collections.Generic;
+using System.ComponentModel;
+using System.Diagnostics;
+using System.IO;
+using System.Linq;
+using System.Reflection;
+using System.Text;
+using System.Text.RegularExpressions;
+using System.Threading;
+using System.Threading.Tasks;
+using Tools.DotNETCommon;
+
+namespace BuildAgent.Run
+{
+	/// <summary>
+	/// Executes an external command and parses the output. Echoes the output to the calling process, and parses structured errors that can be used for posting build health info to UnrealGameSync.
+	/// </summary>
+	[ProgramMode("Run", "Executes a command, processing stdout for structured errors.")]
+	class RunMode : ProgramMode
+	{
+		[CommandLine("-Input=")]
+		[Description("Log file to parse rather than executing an external program.")]
+		FileReference InputFile = null;
+
+		[CommandLine("-Ignore=")]
+		[Description("Path to a file containing error patterns to ignore, one regex per line.")]
+		FileReference IgnorePatternsFile = null;
+
+		[CommandLine]
+		[Description("The program to run.")]
+		FileReference Program = null;
+
+		[CommandLine("-Arguments=")]
+		[Description("Specifies command line arguments to pass to the child process.")]
+		string Arguments = "";
+
+		[CommandLine]
+		[Description("Amount of time to leave before killing the child process.")]
+		TimeSpan? Timeout = null;
+
+		[CommandLine("-DebugListener")]
+		[Description("Enables the debug listener")]
+		bool bDebugListener = false;
+
+		[CommandLine("-ECListener")]
+		[Description("Enables the ElectricCommander listener")]
+		bool bElectricCommanderListener = false;
+
+		public override void Configure(CommandLineArguments Arguments)
+		{
+			base.Configure(Arguments);
+
+			if (InputFile == null && Program == null)
+			{
+				throw new CommandLineArgumentException(String.Format("Either -{0}=... or -{1}=... must be specified.", nameof(InputFile), nameof(Program)));
+			}
+
+			if (!bElectricCommanderListener && !String.IsNullOrEmpty(Environment.GetEnvironmentVariable("COMMANDER_JOBSTEP")))
+			{
+				bElectricCommanderListener = true;
+			}
+		}
+
+		public override void Execute()
+		{
+			// Auto-register all the known matchers in this assembly
+			List<IErrorMatcher> Matchers = new List<IErrorMatcher>();
+			foreach (Type Type in Assembly.GetExecutingAssembly().GetTypes())
+			{
+				if (Type.GetCustomAttribute<AutoRegisterAttribute>() != null)
+				{
+					object Instance = Activator.CreateInstance(Type);
+					if (typeof(IErrorMatcher).IsAssignableFrom(Type))
+					{
+						Matchers.Add((IErrorMatcher)Instance);
+					}
+					else
+					{
+						throw new Exception(String.Format("Unable to auto-register object of type {0}", Type.Name));
+					}
+				}
+			}
+
+			// Read all the ignore patterns
+			List<string> IgnorePatterns = new List<string>();
+			if(IgnorePatternsFile != null)
+			{
+				if (!FileReference.Exists(IgnorePatternsFile))
+				{
+					throw new FatalErrorException("Unable to read '{0}", IgnorePatternsFile);
+				}
+
+				// Read all the ignore patterns
+				string[] Lines = FileReference.ReadAllLines(IgnorePatternsFile);
+				foreach (string Line in Lines)
+				{
+					string TrimLine = Line.Trim();
+					if (TrimLine.Length > 0 && !TrimLine.StartsWith("#"))
+					{
+						IgnorePatterns.Add(TrimLine);
+					}
+				}
+			}
+
+			// Create the output listeners
+			List<IErrorListener> Listeners = new List<IErrorListener>();
+			try
+			{
+				if (bDebugListener)
+				{
+					Listeners.Add(new DebugOutputListener());
+				}
+				if (bElectricCommanderListener)
+				{
+					Listeners.Add(new ElectricCommanderListener());
+				}
+
+				// Process the input
+				if (InputFile != null)
+				{
+					if (!FileReference.Exists(InputFile))
+					{
+						throw new FatalErrorException("Specified input file '{0}' does not exist", InputFile);
+					}
+
+					using (StreamReader Reader = new StreamReader(InputFile.FullName))
+					{
+						LineFilter Filter = new LineFilter(() => Reader.ReadLine());
+						ProcessErrors(Filter.ReadLine, Matchers, IgnorePatterns, Listeners);
+					}
+				}
+				else
+				{
+					CancellationTokenSource CancellationTokenSource = new CancellationTokenSource();
+					if (Timeout.HasValue)
+					{
+						CancellationTokenSource.CancelAfter(Timeout.Value);
+					}
+
+					CancellationToken CancellationToken = CancellationTokenSource.Token;
+					using (ManagedProcess Process = new ManagedProcess(null, Program.FullName, Arguments, null, null, null, ProcessPriorityClass.Normal))
+					{
+						Func<string> ReadLine = new LineFilter(() => ReadProcessLine(Process, CancellationToken)).ReadLine;
+						ProcessErrors(ReadLine, Matchers, IgnorePatterns, Listeners);
+					}
+
+					ProcessUtils.TerminateChildProcesses();
+				}
+			}
+			finally
+			{
+				foreach(IErrorListener Listener in Listeners)
+				{
+					Listener.Dispose();
+				}
+			}
+		}
+
+		/// <summary>
+		/// Reads a line of output from the given process
+		/// </summary>
+		/// <param name="Process">The process to read from</param>
+		/// <param name="CancellationToken">Cancellation token for when the timeout expires</param>
+		/// <returns>The line that was read</returns>
+		static string ReadProcessLine(ManagedProcess Process, CancellationToken CancellationToken)
+		{
+			string Line;
+			if (Process.TryReadLine(out Line, CancellationToken))
+			{
+				Log.TraceInformation("{0}", Line);
+			}
+			return Line;
+		}
+
+		/// <summary>
+		/// Process all the errors obtained by calling the ReadLine() function, and forward them to an array of listeners
+		/// </summary>
+		/// <param name="ReadLine">Delegate used to retrieve each output line</param>
+		/// <param name="Matchers">List of matchers to run against the text</param>
+		/// <param name="IgnorePatterns">List of patterns to ignore</param>
+		/// <param name="Listeners">Set of listeners for processing the errors</param>
+		static void ProcessErrors(Func<string> ReadLine, List<IErrorMatcher> Matchers, List<string> IgnorePatterns, List<IErrorListener> Listeners)
+		{
+			System.Text.RegularExpressions.Regex.CacheSize = 1000;
+
+			LineBuffer Buffer = new LineBuffer(ReadLine, 50);
+			ReadOnlyLineBuffer ReadOnlyBuffer = new ReadOnlyLineBuffer(Buffer);
+
+			Stopwatch Timer = Stopwatch.StartNew();
+			while (Buffer[0] != null)
+			{
+				if ((Buffer.CurrentLineNumber % 1000) == 0)
+				{
+					Log.TraceInformation("LINE: {0} ({1}s)", Buffer.CurrentLineNumber, (int)Timer.Elapsed.TotalSeconds);
+				}
+
+				// Try to match an error
+				ErrorMatch Error = null;
+				foreach (IErrorMatcher Matcher in Matchers)
+				{
+					ErrorMatch NewError = Matcher.Match(ReadOnlyBuffer);
+					if (NewError != null && (Error == null || NewError.Priority > Error.Priority))
+					{
+						Error = NewError;
+					}
+				}
+
+				// If we did match something, check if it's not negated by an ignore pattern. We typically have relatively few errors and many more ignore patterns than matchers, so it's quicker 
+				// to check them in response to an identified error than to treat them as matchers of their own.
+				if (Error != null)
+				{
+					foreach (string IgnorePattern in IgnorePatterns)
+					{
+						if(Regex.IsMatch(Buffer[0], IgnorePattern))
+						{
+							Error = null;
+							break;
+						}
+					}
+				}
+
+				// Report the error to the listeners
+				int AdvanceLines = 1;
+				if (Error != null)
+				{
+					foreach (IErrorListener Listener in Listeners)
+					{
+						Listener.OnErrorMatch(Error);
+					}
+					AdvanceLines = Error.MaxLineNumber + 1 - Buffer.CurrentLineNumber;
+				}
+
+				// Move forwards
+				Buffer.Advance(AdvanceLines);
+			}
+		}
+	}
+}
diff --git a/Engine/Source/Programs/BuildAgent/Workspace/CleanMode.cs b/Engine/Source/Programs/BuildAgent/Workspace/CleanMode.cs
new file mode 100644
--- /dev/null
+++ b/Engine/Source/Programs/BuildAgent/Workspace/CleanMode.cs
@@ -0,0 +1,20 @@
+// Copyright 1998-2019 Epic Games, Inc. All Rights Reserved.
+
+using System;
+using System.Collections.Generic;
+using System.Linq;
+using System.Text;
+using System.Threading.Tasks;
+using BuildAgent.Workspace.Common;
+
+namespace BuildAgent.Workspace
+{
+	[ProgramMode("Clean", "Cleans all modified files from the workspace.")]
+	class CleanMode : WorkspaceMode
+	{
+		protected override void Execute(Repository Repo)
+		{
+			Repo.Clean();
+		}
+	}
+}
diff --git a/Engine/Source/Programs/BuildAgent/Workspace/ClearMode.cs b/Engine/Source/Programs/BuildAgent/Workspace/ClearMode.cs
new file mode 100644
--- /dev/null
+++ b/Engine/Source/Programs/BuildAgent/Workspace/ClearMode.cs
@@ -0,0 +1,20 @@
+// Copyright 1998-2019 Epic Games, Inc. All Rights Reserved.
+
+using BuildAgent.Workspace.Common;
+using System;
+using System.Collections.Generic;
+using System.Linq;
+using System.Text;
+using System.Threading.Tasks;
+
+namespace BuildAgent.Workspace
+{
+	[ProgramMode("Clear", "Empties the staging directory of any files, returning them to the cache")]
+	class ClearMode : WorkspaceMode
+	{
+		protected override void Execute(Repository Repo)
+		{
+			Repo.Clear();
+		}
+	}
+}
diff --git a/Engine/Source/Programs/BuildAgent/Workspace/Common/FileContentId.cs b/Engine/Source/Programs/BuildAgent/Workspace/Common/FileContentId.cs
new file mode 100644
--- /dev/null
+++ b/Engine/Source/Programs/BuildAgent/Workspace/Common/FileContentId.cs
@@ -0,0 +1,77 @@
+// Copyright 1998-2019 Epic Games, Inc. All Rights Reserved.
+
+using System;
+using System.Collections.Generic;
+using System.IO;
+using System.Linq;
+using System.Text;
+using System.Threading.Tasks;
+using Tools.DotNETCommon;
+
+namespace BuildAgent.Workspace.Common
+{
+	class FileContentId : IBinarySerializable
+	{
+		const int DigestLength = 16;
+
+		public byte[] Digest;
+		public string Type;
+
+		public FileContentId(byte[] Digest, string Type)
+		{
+			this.Digest = Digest;
+			this.Type = Type;
+		}
+
+		public FileContentId(BinaryReader Reader)
+		{
+			this.Digest = Reader.ReadBytes(DigestLength);
+			this.Type = Reader.ReadString();
+		}
+
+		public void Write(BinaryWriter Writer)
+		{
+			Writer.Write(Digest);
+			Writer.Write(Type);
+		}
+
+		public override bool Equals(object Other)
+		{
+			FileContentId OtherContentId = Other as FileContentId;
+			if(OtherContentId == null)
+			{
+				return false;
+			}
+
+			for(int Idx = 0; Idx < DigestLength; Idx++)
+			{
+				if(OtherContentId.Digest[Idx] != Digest[Idx])
+				{
+					return false;
+				}
+			}
+
+			if(OtherContentId.Type != Type)
+			{
+				return false;
+			}
+
+			return true;
+		}
+
+		public override int GetHashCode()
+		{
+			int HashCode = Type.GetHashCode();
+			for(int Idx = 0; Idx < DigestLength; Idx++)
+			{
+				HashCode = (HashCode * 31) + Digest[Idx];
+			}
+			return HashCode;
+		}
+
+		public override string ToString()
+		{
+			return String.Format("{0} ({1})", StringUtils.FormatHexString(Digest), Type);
+		}
+	}
+}
diff --git a/Engine/Source/Programs/BuildAgent/Workspace/Common/Repository.cs b/Engine/Source/Programs/BuildAgent/Workspace/Common/Repository.cs
new file mode 100644
--- /dev/null
+++ b/Engine/Source/Programs/BuildAgent/Workspace/Common/Repository.cs
@@ -0,0 +1,1552 @@
+// Copyright 1998-2019 Epic Games, Inc. All Rights Reserved.
+
+using System;
+using System.Collections.Concurrent;
+using System.Collections.Generic;
+using System.Diagnostics;
+using System.IO;
+using System.Linq;
+using System.Reflection;
+using System.Text;
+using System.Threading.Tasks;
+using Tools.DotNETCommon;
+using Tools.DotNETCommon.Perforce;
+
+namespace BuildAgent.Workspace.Common
+{
+	/// <summary>
+	/// Represents a repository of streams and cached data
+	/// </summary>
+	class Repository
+	{
+		/// <summary>
+		/// The current transaction state. Used to determine whether a repository needs to be cleaned on startup.
+		/// </summary>
+		enum RepositoryState
+		{
+			Dirty,
+			Clean,
+		}
+
+		/// <summary>
+		/// The file signature and version. Update this to introduce breaking changes and ignore old repositories.
+		/// </summary>
+		const int CurrentSignature = ('W' << 24) | ('T' << 16) | 1;
+
+		/// <summary>
+		/// The current revision number.
+		/// </summary>
+		const int CurrentRevision = 1;
+
+		/// <summary>
+		/// Constant for syncing the latest change number
+		/// </summary>
+		public const int LatestChangeNumber = -1;
+
+		/// <summary>
+		/// Name of the signature file for a repository. This 
+		/// </summary>
+		const string SignatureFileName = "Repository.sig";
+
+		/// <summary>
+		/// Name of the main data file for a repository
+		/// </summary>
+		const string DataFileName = "Repository.dat";
+
+		/// <summary>
+		/// Name of the server and port for 
+		/// </summary>
+		string ServerAndPort;
+
+		/// <summary>
+		/// Name of the Perforce user
+		/// </summary>
+		string UserName;
+
+		/// <summary>
+		/// Name of the host
+		/// </summary>
+		string HostName;
+
+		/// <summary>
+		/// Incrementing number assigned to sequential operations that modify files. Used to age out files in the cache.
+		/// </summary>
+		uint NextSequenceNumber;
+
+		/// <summary>
+		/// Whether a repair operation should be run on this workspace. Set whenever the state may be inconsistent.
+		/// </summary>
+		bool bRequiresRepair;
+
+		/// <summary>
+		/// Connection to the Perforce server
+		/// </summary>
+		PerforceConnection Perforce;
+
+		/// <summary>
+		/// The root directory for the stash
+		/// </summary>
+		readonly DirectoryReference BaseDir;
+
+		/// <summary>
+		/// Root directory for storing cache files
+		/// </summary>
+		readonly DirectoryReference CacheDir;
+
+		/// <summary>
+		/// Root directory for storing workspace files
+		/// </summary>
+		readonly DirectoryReference WorkspaceDir;
+
+		/// <summary>
+		/// Set of clients that we're created. Used to avoid updating multiple times during one run.
+		/// </summary>
+		HashSet<string> ClientNames = new HashSet<string>();
+
+		/// <summary>
+		/// Set of unique cache entries. We use this to ensure new names in the cache are unique.
+		/// </summary>
+		HashSet<ulong> CacheEntries = new HashSet<ulong>();
+
+		/// <summary>
+		/// List of all the staged files
+		/// </summary>
+		WorkspaceDirectoryInfo Workspace;
+
+		/// <summary>
+		/// All the files which are currently being tracked
+		/// </summary>
+		Dictionary<FileContentId, TrackedFileInfo> ContentIdToTrackedFile = new Dictionary<FileContentId, TrackedFileInfo>();
+
+		/// <summary>
+		/// Constructor
+		/// </summary>
+		/// <param name="ServerAndPort">Perforce server and port to connect to</param>
+		/// <param name="UserName">Name of the user to login to Perforce with</param>
+		/// <param name="HostName">Name of the current host</param>
+		/// <param name="NextSequenceNumber">The next sequence number for operations</param>
+		/// <param name="BaseDir">The root directory for the stash</param>
+		private Repository(string ServerAndPort, string UserName, string HostName, uint NextSequenceNumber, DirectoryReference BaseDir)
+		{
+			// Make sure all the fields are valid
+			if(UserName == null || HostName == null)
+			{
+				PerforceConnection Perforce = CreatePerforceConnection(ServerAndPort, UserName, null);
+
+				InfoRecord Info = Perforce.Info(InfoOptions.ShortOutput).Data;
+				if(UserName == null)
+				{
+					UserName = Info.UserName;
+				}
+				if(HostName == null)
+				{
+					HostName = Info.ClientHost;
+				}
+			}
+
+			// Save the Perforce settings
+			this.ServerAndPort = ServerAndPort;
+			this.UserName = UserName;
+			this.HostName = HostName;
+			this.NextSequenceNumber = NextSequenceNumber;
+
+			// Create the perforce connection
+			Perforce = CreatePerforceConnection(ServerAndPort, UserName, null);
+
+			// Get all the directories
+			this.BaseDir = BaseDir;
+			DirectoryReference.CreateDirectory(BaseDir);
+
+			this.CacheDir = DirectoryReference.Combine(BaseDir, "Cache");
+			DirectoryReference.CreateDirectory(CacheDir);
+
+			this.WorkspaceDir = DirectoryReference.Combine(BaseDir, "Sync");
+			DirectoryReference.CreateDirectory(WorkspaceDir);
+
+			// Create the workspace
+			this.Workspace = new WorkspaceDirectoryInfo(WorkspaceDir);
+		}
+
+		/// <summary>
+		/// Creates a repository at the given location
+		/// </summary>
+		/// <param name="ServerAndPort">The Perforce server and port</param>
+		/// <param name="UserName">The Perforce username to connect with</param>
+		/// <param name="BaseDir">The base directory for the repository</param>
+		/// <returns>New repository instance</returns>
+		public static Repository Create(string ServerAndPort, string UserName, DirectoryReference BaseDir)
+		{
+			Log.TraceInformation("Creating repository at {0}...", BaseDir);
+
+			DirectoryReference.CreateDirectory(BaseDir);
+			FileUtils.ForceDeleteDirectoryContents(BaseDir);
+
+			Repository Repo = new Repository(ServerAndPort, UserName, null, 1, BaseDir);
+			Repo.Save(RepositoryState.Clean);
+			Repo.CreateCacheHierarchy();
+
+			FileReference SignatureFile = FileReference.Combine(BaseDir, SignatureFileName);
+			using(BinaryWriter Writer = new BinaryWriter(File.Open(SignatureFile.FullName, FileMode.Create, FileAccess.Write, FileShare.Read)))
+			{
+				Writer.Write(CurrentSignature);
+			}
+
+			return Repo;
+		}
+
+		/// <summary>
+		/// Tests whether a repository exists in the given directory
+		/// </summary>
+		/// <param name="BaseDir"></param>
+		/// <returns></returns>
+		public static bool Exists(DirectoryReference BaseDir)
+		{
+			FileReference SignatureFile = FileReference.Combine(BaseDir, SignatureFileName);
+			if(FileReference.Exists(SignatureFile))
+			{
+				using(BinaryReader Reader = new BinaryReader(File.Open(SignatureFile.FullName, FileMode.Open, FileAccess.Read, FileShare.Read)))
+				{
+					int Signature = Reader.ReadInt32();
+					if(Signature == CurrentSignature)
+					{
+						return true;
+					}
+				}
+			}
+			return false;
+		}
+
+		/// <summary>
+		/// Loads a repository from disk
+		/// </summary>
+		/// <param name="ServerAndPort">The Perforce server and port</param>
+		/// <param name="UserName">The Perforce username to connect with</param>
+		/// <param name="BaseDir">The base directory for the repository</param>
+		public static Repository Load(string ServerAndPort, string UserName, DirectoryReference BaseDir)
+		{
+			if(!Exists(BaseDir))
+			{
+				throw new FatalErrorException("No valid repository found at {0}", BaseDir);
+			}
+
+			FileReference DataFile = FileReference.Combine(BaseDir, DataFileName);
+			RestoreBackup(DataFile);
+
+			Repository Repo;
+			using(BinaryReader Reader = new BinaryReader(File.Open(DataFile.FullName, FileMode.Open, FileAccess.Read, FileShare.Read)))
+			{
+				int Revision = Reader.ReadInt32();
+				if(Revision != CurrentRevision)
+				{
+					throw new FatalErrorException("Unsupported data format (revision {0}, expected {1})", Revision, CurrentRevision);
+				}
+
+				bool bRequiresRepair = Reader.ReadBoolean();
+				Reader.ReadString(); // Originally client name; now ignored
+				uint NextSequenceNumber = Reader.ReadUInt32();
+
+				Repo = new Repository(ServerAndPort, UserName, null, NextSequenceNumber, BaseDir);
+				Repo.bRequiresRepair = bRequiresRepair;
+
+				int NumTrackedFiles = Reader.ReadInt32();
+				for(int Idx = 0; Idx < NumTrackedFiles; Idx++)
+				{
+					TrackedFileInfo TrackedFile = new TrackedFileInfo(Repo.CacheDir, Reader);
+					Repo.ContentIdToTrackedFile.Add(TrackedFile.ContentId, TrackedFile);
+					Repo.CacheEntries.Add(TrackedFile.CacheId);
+				}
+
+				Repo.Workspace.Read(Reader);
+			}
+
+			Repo.RunOptionalRepair();
+			return Repo;
+		}
+
+		/// <summary>
+		/// Save the state of the repository
+		/// </summary>
+		private void Save(RepositoryState State)
+		{
+			FileReference DataFile = FileReference.Combine(BaseDir, DataFileName);
+
+			BeginTransaction(DataFile);
+			using(BinaryWriter Writer = new BinaryWriter(File.Open(DataFile.FullName, FileMode.Create, FileAccess.Write, FileShare.Read)))
+			{
+				Writer.Write(CurrentRevision);
+				Writer.Write(bRequiresRepair || (State != RepositoryState.Clean));
+				Writer.Write(""); // Originally client name; now ignored
+				Writer.Write(NextSequenceNumber);
+
+				Writer.Write(ContentIdToTrackedFile.Count);
+				foreach(TrackedFileInfo TrackedFile in ContentIdToTrackedFile.Values)
+				{
+					TrackedFile.Write(Writer);
+				}
+
+				Workspace.Write(Writer);
+			}
+			CompleteTransaction(DataFile);
+		}
+
+		/// <summary>
+		/// Creates a Perforce connection for the given settings
+		/// </summary>
+		/// <param name="ServerAndPort">The name and port for the server</param>
+		/// <param name="UserName">The username to connect to the server with</param>
+		/// <param name="ClientName">Name of the client to use when syncing files</param>
+		/// <returns>New Perforce connection instance</returns>
+		static PerforceConnection CreatePerforceConnection(string ServerAndPort, string UserName, string ClientName)
+		{
+			// Create the Perforce connection
+			List<string> GlobalOptions = new List<string>();
+			GlobalOptions.Add("-zprog=WorkspaceTool");
+			GlobalOptions.Add(String.Format("-zversion={0}", Assembly.GetExecutingAssembly().GetName().Version.ToString()));
+			if(ServerAndPort != null)
+			{
+				GlobalOptions.Add(String.Format("-p {0}", ServerAndPort));
+			}
+			if(UserName != null)
+			{
+				GlobalOptions.Add(String.Format("-u {0}", UserName));
+			}
+			if(ClientName != null)
+			{
+				GlobalOptions.Add(String.Format("-c {0}", ClientName));
+			}
+			return new PerforceConnection(String.Join(" ", GlobalOptions));
+		}
+
+		#region Commands
+
+		/// <summary>
+		/// Cleans the current workspace
+		/// </summary>
+		public void Clean()
+		{
+			Stopwatch Timer = Stopwatch.StartNew();
+
+			Log.TraceInformation("Cleaning workspace...");
+			using(LogIndentScope Indent = new LogIndentScope("  "))
+			{
+				CleanInternal();
+			}
+
+			Log.TraceInformation("Completed in {0:0.0}s", Timer.Elapsed.TotalSeconds);
+		}
+
+		/// <summary>
+		/// Cleans the current workspace
+		/// </summary>
+		private void CleanInternal()
+		{
+			FileInfo[] FilesToDelete;
+			DirectoryInfo[] DirectoriesToDelete;
+			using(LogStatusScope Status = new LogStatusScope("Finding files to clean..."))
+			{
+				Stopwatch Timer = Stopwatch.StartNew();
+
+				Workspace.Refresh(out FilesToDelete, out DirectoriesToDelete);
+
+				Status.SetProgress("({0:0.0}s)", Timer.Elapsed.TotalSeconds);
+			}
+
+			if(FilesToDelete.Length > 0 || DirectoriesToDelete.Length > 0)
+			{
+				List<string> Paths = new List<string>();
+				Paths.AddRange(DirectoriesToDelete.Select(x => String.Format("/{0}/...", new DirectoryReference(x).MakeRelativeTo(WorkspaceDir).Replace(Path.DirectorySeparatorChar, '/'))));
+				Paths.AddRange(FilesToDelete.Select(x => String.Format("/{0}", new FileReference(x).MakeRelativeTo(WorkspaceDir).Replace(Path.DirectorySeparatorChar, '/'))));
+
+				foreach(string Path in Paths.OrderBy(x => x))
+				{
+					Log.TraceInformation("  {0}", Path);
+				}
+
+				using(LogStatusScope Scope = new LogStatusScope("Cleaning files..."))
+				{
+					Stopwatch Timer = Stopwatch.StartNew();
+					using(ThreadPoolWorkQueue Queue = new ThreadPoolWorkQueue())
+					{
+						foreach(FileInfo FileToDelete in FilesToDelete)
+						{
+							Queue.Enqueue(() => FileUtils.ForceDeleteFile(FileToDelete));
+						}
+						foreach(DirectoryInfo DirectoryToDelete in DirectoriesToDelete)
+						{
+							Queue.Enqueue(() => FileUtils.ForceDeleteDirectory(DirectoryToDelete));
+						}
+					}
+					Scope.SetProgress("({0:0.0}s)", Timer.Elapsed.TotalSeconds);
+				}
+
+				Save(RepositoryState.Clean);
+			}
+		}
+
+		/// <summary>
+		/// Empties the staging directory of any staged files
+		/// </summary>
+		public void Clear()
+		{
+			Stopwatch Timer = Stopwatch.StartNew();
+
+			Log.TraceInformation("Clearing workspace...");
+			using(LogIndentScope Indent = new LogIndentScope("  "))
+			{
+				CleanInternal();
+				RemoveFilesFromWorkspace(new StreamDirectoryInfo());
+				Save(RepositoryState.Clean);
+			}
+
+			Log.TraceInformation("Completed in {0:0.0}s", Timer.Elapsed.TotalSeconds);
+		}
+
+		/// <summary>
+		/// Dumps the contents of the repository to the log for analysis
+		/// </summary>
+		public void Dump()
+		{
+			Stopwatch Timer = Stopwatch.StartNew();
+
+			Log.TraceInformation("Dumping repository to log...");
+
+			WorkspaceFileInfo[] WorkspaceFiles = Workspace.GetFiles().OrderBy(x => x.GetLocation().FullName).ToArray();
+			if(WorkspaceFiles.Length > 0)
+			{
+				Log.TraceLog("  Workspace:");
+				foreach(WorkspaceFileInfo File in WorkspaceFiles)
+				{
+					Log.TraceLog("    {0,-128} [{1,-48}] [{2,20:n0}] [{3,20}]{4}", File.GetClientPath(), File.ContentId, File.Length, File.LastModifiedTicks, File.bReadOnly? "" : " [ writable ]");
+				}
+			}
+
+			if(ContentIdToTrackedFile.Count > 0)
+			{
+				Log.TraceLog("  Cache:");
+				foreach(KeyValuePair<FileContentId, TrackedFileInfo> Pair in ContentIdToTrackedFile)
+				{
+					Log.TraceLog("    {0,-128} [{1,-48}] [{2,20:n0}] [{3,20}]{4}", Pair.Value.GetLocation(), Pair.Key, Pair.Value.Length, Pair.Value.LastModifiedTicks, Pair.Value.bReadOnly? "" : "[ writable ]");
+				}
+			}
+
+			Log.TraceInformation("Completed in {0:0.0}s", Timer.Elapsed.TotalSeconds);
+		}
+
+		/// <summary>
+		/// Checks the integrity of the cache
+		/// </summary>
+		public void Repair()
+		{
+			using(LogStatusScope Status = new LogStatusScope("Checking cache..."))
+			{
+				// Make sure all the folders exist in the cache
+				CreateCacheHierarchy();
+
+				// Check that all the files in the cache appear as we expect them to
+				List<TrackedFileInfo> TrackedFiles = ContentIdToTrackedFile.Values.ToList();
+				foreach(TrackedFileInfo TrackedFile in TrackedFiles)
+				{
+					if(!TrackedFile.CheckIntegrity())
+					{
+						RemoveTrackedFile(TrackedFile);
+					}
+				}
+
+				// Clear the repair flag
+				bRequiresRepair = false;
+
+				Save(RepositoryState.Clean);
+
+				Status.SetProgress("Done");
+			}
+		}
+
+		/// <summary>
+		/// Cleans the current workspace
+		/// </summary>
+		public void Revert(string ClientName)
+		{
+			Stopwatch Timer = Stopwatch.StartNew();
+
+			PerforceConnection Client = CreatePerforceConnection(ServerAndPort, UserName, ClientName);
+			RevertInternal(Client, ClientName);
+
+			Log.TraceInformation("Completed in {0:0.0}s", Timer.Elapsed.TotalSeconds);
+		}
+
+		/// <summary>
+		/// Checks the bRequiresRepair flag, and repairs/resets it if set.
+		/// </summary>
+		private void RunOptionalRepair()
+		{
+			if(bRequiresRepair)
+			{
+				Repair();
+			}
+		}
+
+		/// <summary>
+		/// Shrink the size of the cache to the given size
+		/// </summary>
+		/// <param name="MaxSize">The maximum cache size, in bytes</param>
+		public void Purge(long MaxSize)
+		{
+			Log.TraceInformation("Purging cache (limit {0:n0} bytes)...", MaxSize);
+			using(new LogIndentScope("  "))
+			{
+				List<TrackedFileInfo> CachedFiles = ContentIdToTrackedFile.Values.OrderBy(x => x.SequenceNumber).ToList();
+
+				int NumRemovedFiles = 0;
+				long TotalSize = CachedFiles.Sum(x => x.Length);
+
+				while(MaxSize < TotalSize && NumRemovedFiles < CachedFiles.Count)
+				{
+					TrackedFileInfo File = CachedFiles[NumRemovedFiles];
+
+					RemoveTrackedFile(File);
+					TotalSize -= File.Length;
+
+					NumRemovedFiles++;
+				}
+
+				Save(RepositoryState.Clean);
+
+				Log.TraceInformation("{0} files removed, {1} files remaining, new size {2:n0} bytes.", NumRemovedFiles, CachedFiles.Count - NumRemovedFiles, TotalSize);
+			}
+		}
+
+		/// <summary>
+		/// Configures the client for the given stream
+		/// </summary>
+		/// <param name="ClientName">Name of the client</param>
+		/// <param name="StreamName">Name of the stream</param>
+		public void Setup(string ClientName, string StreamName)
+		{
+			UpdateClient(ClientName, StreamName);
+		}
+
+		/// <summary>
+		/// Prints stats showing coherence between different streams
+		/// </summary>
+		public void Stats(string ClientName, List<string> StreamNames, List<string> Filters)
+		{
+			Log.TraceInformation("Finding stats for {0} streams", StreamNames.Count);
+			using (new LogIndentScope("  "))
+			{
+				// Update the list of files in each stream
+				Tuple<int, StreamDirectoryInfo>[] StreamState = new Tuple<int, StreamDirectoryInfo>[StreamNames.Count];
+				for (int Idx = 0; Idx < StreamNames.Count; Idx++)
+				{
+					string StreamName = StreamNames[Idx];
+					Log.TraceInformation("Finding contents of {0}:", StreamName);
+
+					using (new LogIndentScope("  "))
+					{
+						ClientNames.Remove(ClientName); // Force the client to be updated
+
+						PerforceConnection Client = UpdateClient(ClientName, StreamName);
+
+						int ChangeNumber = GetLatestClientChange(Client, ClientName);
+						Log.TraceInformation("Latest change is CL {0}", ChangeNumber);
+
+						RevertInternal(Client, ClientName);
+						ClearClientHaveTable(Client, ClientName);
+						UpdateClientHaveTable(Client, ClientName, ChangeNumber, Filters);
+
+						StreamDirectoryInfo Contents = FindClientContents(Client, ClientName, ChangeNumber, false);
+						StreamState[Idx] = Tuple.Create(ChangeNumber, Contents);
+
+						GC.Collect();
+					}
+				}
+
+				// Find stats for 
+				using (LogStatusScope Scope = new LogStatusScope("Finding usage stats..."))
+				{
+					Stopwatch Timer = Stopwatch.StartNew();
+
+					// Find the set of files in each stream
+					HashSet<FileContentId>[] FilesInStream = new HashSet<FileContentId>[StreamNames.Count];
+					for (int Idx = 0; Idx < StreamNames.Count; Idx++)
+					{
+						List<StreamFileInfo> Files = StreamState[Idx].Item2.GetFiles();
+						FilesInStream[Idx] = new HashSet<FileContentId>(Files.Select(x => x.ContentId));
+					}
+
+					// Build a table showing amount of unique content in each stream
+					string[,] Cells = new string[StreamNames.Count + 1, StreamNames.Count + 1];
+					Cells[0, 0] = "";
+					for (int Idx = 0; Idx < StreamNames.Count; Idx++)
+					{
+						Cells[Idx + 1, 0] = StreamNames[Idx];
+						Cells[0, Idx + 1] = StreamNames[Idx];
+					}
+
+					// Populate the table
+					for (int RowIdx = 0; RowIdx < StreamNames.Count; RowIdx++)
+					{
+						List<StreamFileInfo> Files = StreamState[RowIdx].Item2.GetFiles();
+						for (int ColIdx = 0; ColIdx < StreamNames.Count; ColIdx++)
+						{
+							long DiffSize = Files.Where(x => !FilesInStream[ColIdx].Contains(x.ContentId)).Sum(x => x.Length);
+							Cells[RowIdx + 1, ColIdx + 1] = String.Format("{0:0.0}mb", DiffSize / (1024.0 * 1024.0));
+						}
+					}
+
+					// Find the width of each row
+					int[] ColWidths = new int[StreamNames.Count + 1];
+					for (int ColIdx = 0; ColIdx < StreamNames.Count + 1; ColIdx++)
+					{
+						for (int RowIdx = 0; RowIdx < StreamNames.Count + 1; RowIdx++)
+						{
+							ColWidths[ColIdx] = Math.Max(ColWidths[ColIdx], Cells[RowIdx, ColIdx].Length);
+						}
+					}
+
+					// Print the table
+					Log.TraceInformation("");
+					Log.TraceInformation("Each row shows the size of files in a stream which are unique to that stream compared to each column:");
+					Log.TraceInformation("");
+					for (int RowIdx = 0; RowIdx < StreamNames.Count + 1; RowIdx++)
+					{
+						StringBuilder Row = new StringBuilder();
+						for (int ColIdx = 0; ColIdx < StreamNames.Count + 1; ColIdx++)
+						{
+							string Cell = Cells[RowIdx, ColIdx];
+							Row.Append(' ', ColWidths[ColIdx] - Cell.Length);
+							Row.Append(Cell);
+							Row.Append(" | ");
+						}
+						Log.TraceInformation("{0}", Row.ToString());
+					}
+					Log.TraceInformation("");
+
+					Scope.SetProgress("({0:0.0}s)", Timer.Elapsed.TotalSeconds);
+				}
+			}
+		}
+
+		/// <summary>
+		/// Prints information about the repository state
+		/// </summary>
+		public void Status()
+		{
+			// Print size stats
+			Log.TraceInformation("Cache contains {0:n0} files, {1:n1}mb", ContentIdToTrackedFile.Count, ContentIdToTrackedFile.Values.Sum(x => x.Length) / (1024.0 * 1024.0));
+			Log.TraceInformation("Stage contains {0:n0} files, {1:n1}mb", Workspace.GetFiles().Count, Workspace.GetFiles().Sum(x => x.Length) / (1024.0 * 1024.0));
+
+			// Print the contents of the workspace
+			string[] Differences = Workspace.FindDifferences();
+			if(Differences.Length > 0)
+			{
+				Log.TraceInformation("Local changes:");
+				foreach(string Difference in Differences)
+				{
+					if(Difference.StartsWith("+"))
+					{
+						Console.ForegroundColor = ConsoleColor.Green;
+					}
+					else if(Difference.StartsWith("-"))
+					{
+						Console.ForegroundColor = ConsoleColor.Red;
+					}
+					else if(Difference.StartsWith("!"))
+					{
+						Console.ForegroundColor = ConsoleColor.Yellow;
+					}
+					else
+					{
+						Console.ResetColor();
+					}
+					Log.TraceInformation("  {0}", Difference);
+				}
+				Console.ResetColor();
+			}
+		}
+
+		/// <summary>
+		/// Switches to the given stream
+		/// </summary>
+		/// <param name="ClientName">Name of the client to sync</param>
+		/// <param name="StreamName">Name of the stream to sync</param>
+		/// <param name="ChangeNumber">Changelist number to sync. -1 to sync to latest.</param>
+		/// <param name="Filters">List of filters to be applied to the workspace</param>
+		/// <param name="bFakeSync">Whether to simulate the syncing operation rather than actually getting files from the server</param>
+		/// <param name="CacheFile">If set, uses the given file to cache the contents of the workspace. This can improve sync times when multiple machines sync the same workspace.</param>
+		public void Sync(string ClientName, string StreamName, int ChangeNumber, List<string> Filters, bool bFakeSync, FileReference CacheFile)
+		{
+			Stopwatch Timer = Stopwatch.StartNew();
+			if(ChangeNumber == -1)
+			{
+				Log.TraceInformation("Syncing to {0} at latest", StreamName);
+			}
+			else
+			{
+				Log.TraceInformation("Syncing to {0} at CL {1}", StreamName, ChangeNumber);
+			}
+
+			using(LogIndentScope Indent = new LogIndentScope("  "))
+			{
+				// Update the client to the current stream
+				PerforceConnection Client = UpdateClient(ClientName, StreamName);
+
+				// Get the latest change number
+				if(ChangeNumber == -1)
+				{
+					ChangeNumber = GetLatestClientChange(Client, ClientName);
+				}
+
+				// Revert any open files
+				RevertInternal(Client, ClientName);
+
+				// Force the P4 metadata to match up
+				UpdateClientHaveTable(Client, ClientName, ChangeNumber, Filters);
+
+				// Clean the current workspace
+				CleanInternal();
+
+				// Update the state of the current stream, if necessary
+				StreamDirectoryInfo Contents;
+				if(CacheFile == null)
+				{
+					Contents = FindClientContents(Client, ClientName, ChangeNumber, bFakeSync);
+				}
+				else
+				{
+					if(FileReference.Exists(CacheFile))
+					{
+						Contents = LoadClientContents(CacheFile);
+					}
+					else
+					{
+						Contents = FindAndSaveClientContents(Client, ClientName, ChangeNumber, bFakeSync, CacheFile);
+					}
+				}
+
+				// Sync all the appropriate files
+				RemoveFilesFromWorkspace(Contents);
+				AddFilesToWorkspace(Client, Contents, bFakeSync);
+			}
+
+			Log.TraceInformation("Completed in {0:0.0}s", Timer.Elapsed.TotalSeconds);
+		}
+
+		/// <summary>
+		/// Populates the cache with the head revision of the given streams.
+		/// </summary>
+		public void Populate(List<KeyValuePair<string, string>> ClientAndStreamNames, List<string> Filters, bool bFakeSync)
+		{
+			Log.TraceInformation("Populating with {0} streams", ClientAndStreamNames.Count);
+			using(new LogIndentScope("  "))
+			{
+				// Clean the current workspace
+				Clean();
+
+				// Update the list of files in each stream
+				Tuple<int, StreamDirectoryInfo>[] StreamState = new Tuple<int, StreamDirectoryInfo>[ClientAndStreamNames.Count];
+				for(int Idx = 0; Idx < ClientAndStreamNames.Count; Idx++)
+				{
+					string ClientName = ClientAndStreamNames[Idx].Key;
+					string StreamName = ClientAndStreamNames[Idx].Value;
+					Log.TraceInformation("Finding contents of {0}:", StreamName);
+
+					using(new LogIndentScope("  "))
+					{
+						PerforceConnection Client = UpdateClient(ClientName, StreamName);
+
+						int ChangeNumber = GetLatestClientChange(Client, ClientName);
+						Log.TraceInformation("Latest change is CL {0}", ChangeNumber);
+
+						RevertInternal(Client, ClientName);
+						ClearClientHaveTable(Client, ClientName);
+						UpdateClientHaveTable(Client, ClientName, ChangeNumber, Filters);
+
+						StreamDirectoryInfo Contents = FindClientContents(Client, ClientName, ChangeNumber, bFakeSync);
+						StreamState[Idx] = Tuple.Create(ChangeNumber, Contents);
+
+						GC.Collect();
+					}
+				}
+
+				// Remove any files from the workspace not referenced by the first stream. This ensures we can purge things from the cache that we no longer need.
+				if(ClientAndStreamNames.Count > 0)
+				{
+					RemoveFilesFromWorkspace(StreamState[0].Item2);
+				}
+
+				// Shrink the contents of the cache
+				using(LogStatusScope Status = new LogStatusScope("Updating cache..."))
+				{
+					Stopwatch Timer = Stopwatch.StartNew();
+
+					HashSet<FileContentId> CommonContentIds = new HashSet<FileContentId>();
+					Dictionary<FileContentId, long> ContentIdToLength = new Dictionary<FileContentId, long>();
+					for(int Idx = 0; Idx < ClientAndStreamNames.Count; Idx++)
+					{
+						List<StreamFileInfo> Files = StreamState[Idx].Item2.GetFiles();
+						foreach(StreamFileInfo File in Files)
+						{
+							ContentIdToLength[File.ContentId] = File.Length;
+						}
+
+						if(Idx == 0)
+						{
+							CommonContentIds.UnionWith(Files.Select(x => x.ContentId));
+						}
+						else
+						{
+							CommonContentIds.IntersectWith(Files.Select(x => x.ContentId));
+						}
+					}
+
+					List<TrackedFileInfo> TrackedFiles = ContentIdToTrackedFile.Values.ToList();
+					foreach(TrackedFileInfo TrackedFile in TrackedFiles)
+					{
+						if(!ContentIdToLength.ContainsKey(TrackedFile.ContentId))
+						{
+							RemoveTrackedFile(TrackedFile);
+						}
+					}
+
+					GC.Collect();
+
+					double TotalSize = ContentIdToLength.Sum(x => x.Value) / (1024.0 * 1024.0);
+					Status.SetProgress("{0:n1}mb total, {1:n1}mb differences ({2:0.0}s)", TotalSize, TotalSize - CommonContentIds.Sum(x => ContentIdToLength[x]) / (1024.0 * 1024.0), Timer.Elapsed.TotalSeconds);
+				}
+
+				// Sync all the new files
+				for(int Idx = 0; Idx < ClientAndStreamNames.Count; Idx++)
+				{
+					string ClientName = ClientAndStreamNames[Idx].Key;
+					string StreamName = ClientAndStreamNames[Idx].Value;
+					Log.TraceInformation("Syncing files for {0}:", StreamName);
+
+					using(new LogIndentScope("  "))
+					{
+						PerforceConnection Client = UpdateClient(ClientName, StreamName);
+
+						int ChangeNumber = StreamState[Idx].Item1;
+						UpdateClientHaveTable(Client, ClientName, ChangeNumber, Filters);
+
+						StreamDirectoryInfo Contents = StreamState[Idx].Item2;
+						RemoveFilesFromWorkspace(Contents);
+						AddFilesToWorkspace(Client, Contents, bFakeSync);
+					}
+				}
+
+				// Save the new repo state
+				Save(RepositoryState.Clean);
+			}
+		}
+
+		#endregion
+
+		#region Core operations
+
+		/// <summary>
+		/// Sets the stream for the current client
+		/// </summary>
+		/// <param name="ClientName">Name of the client</param>
+		/// <param name="StreamName">New stream for the client</param>
+		private PerforceConnection UpdateClient(string ClientName, string StreamName)
+		{
+			// Create or update the client if it doesn't exist already
+			if(ClientNames.Add(ClientName))
+			{
+				using(LogStatusScope Status = new LogStatusScope("Updating client..."))
+				{
+					Stopwatch Timer = Stopwatch.StartNew();
+
+					ClientRecord Client = new ClientRecord();
+					Client.Name = ClientName;
+					Client.Owner = UserName;
+					Client.Host = HostName;
+					Client.Stream = StreamName;
+					Client.Root = WorkspaceDir.FullName;
+					Client.Type = "partitioned";
+
+					PerforceResponse Response = Perforce.CreateClient(Client);
+					if(!Response.Succeeded)
+					{
+						Perforce.DeleteClient(DeleteClientOptions.None, ClientName);
+						Perforce.CreateClient(Client).RequireSuccess();
+					}
+
+					Status.SetProgress("({0:0.0}s)", Timer.Elapsed.TotalSeconds);
+				}
+			}
+
+			// Update the config file with the name of the client
+			FileReference ConfigFile = FileReference.Combine(BaseDir, "p4.ini");
+			using(StreamWriter Writer = new StreamWriter(ConfigFile.FullName))
+			{
+				Writer.WriteLine("P4CLIENT={0}", ClientName);
+			}
+
+			// Return a Perforce connection for this client
+			return CreatePerforceConnection(ServerAndPort, UserName, ClientName);
+		}
+
+		/// <summary>
+		/// Get the latest change number in the current client
+		/// </summary>
+		/// <param name="Client">The perforce client connection</param>
+		/// <param name="ClientName">Name of the current client</param>
+		/// <returns>The latest submitted change number</returns>
+		private int GetLatestClientChange(PerforceConnection Client, string ClientName)
+		{
+			int ChangeNumber;
+			using(LogStatusScope Status = new LogStatusScope("Finding latest change..."))
+			{
+				Stopwatch Timer = Stopwatch.StartNew();
+
+				PerforceResponseList<ChangesRecord> Changes = Client.Changes(ChangesOptions.None, 1, ChangeStatus.Submitted, String.Format("//{0}/...", ClientName));
+				Changes.RequireSuccess();
+				ChangeNumber = Changes[0].Data.Number;
+
+				Status.SetProgress("CL {0} ({1:0.0}s)", ChangeNumber, Timer.Elapsed.TotalSeconds);
+			}
+			return ChangeNumber;
+		}
+
+		/// <summary>
+		/// Revert all files that are open in the current workspace. Does not replace them with valid revisions.
+		/// </summary>
+		/// <param name="Client">The current client connection</param>
+		/// <param name="ClientName">Name of the current client</param>
+		private void RevertInternal(PerforceConnection Client, string ClientName)
+		{
+			using(LogStatusScope Status = new LogStatusScope("Reverting changes..."))
+			{
+				Stopwatch Timer = Stopwatch.StartNew();
+
+				// Get a list of open files
+				PerforceResponseList<FStatRecord> OpenedFilesResponse = Client.Opened(OpenedOptions.ShortOutput, -1, ClientName, null, 1);
+				OpenedFilesResponse.RequireSuccess();
+
+				// If there are any files, revert them
+				if(OpenedFilesResponse.Any())
+				{
+					Client.Revert(-1, null, RevertOptions.KeepWorkspaceFiles, "//...").RequireSuccess();
+				}
+
+				// Find all the open changes
+				PerforceResponseList<ChangesRecord> Changes = Client.Changes(ChangesOptions.None, ClientName, -1, ChangeStatus.Pending, null);
+				Changes.RequireSuccess();
+
+				// Delete the changelist
+				foreach(ChangesRecord Change in Changes.Data)
+				{
+					// Find a list of shelved changes
+					PerforceResponseList<DescribeRecord> DescribeResponse = Client.Describe(DescribeOptions.Shelved, -1, Change.Number);
+					DescribeResponse.RequireSuccess();
+
+					// Delete the shelved files
+					foreach(DescribeRecord Record in DescribeResponse.Data)
+					{
+						if(Record.Files.Count > 0)
+						{
+							Client.DeleteShelvedFiles(Record.Number).RequireSuccess();
+						}
+					}
+
+					// Delete the changelist
+					Client.DeleteChange(DeleteChangeOptions.None, Change.Number).RequireSuccess();
+				}
+
+				Status.SetProgress("({0:0.0}s)", Timer.Elapsed.TotalSeconds);
+			}
+		}
+
+		/// <summary>
+		/// Clears the have table. This ensures that we'll always fetch the names of files at head revision, which aren't updated otherwise.
+		/// </summary>
+		/// <param name="Client">The client connection</param>
+		/// <param name="ClientName">Name of the current client</param>
+		private void ClearClientHaveTable(PerforceConnection Client, string ClientName)
+		{
+			using(LogStatusScope Scope = new LogStatusScope("Clearing have table..."))
+			{
+				Stopwatch Timer = Stopwatch.StartNew();
+				Client.SyncQuiet(SyncOptions.KeepWorkspaceFiles, -1, String.Format("//{0}/...#0", ClientName)).RequireSuccess();
+				Scope.SetProgress("({0:0.0}s)", Timer.Elapsed.TotalSeconds);
+			}
+		}
+
+		/// <summary>
+		/// Updates the have table to reflect the given stream
+		/// </summary>
+		/// <param name="Client">The client connection</param>
+		/// <param name="ClientName">Name of the current client</param>
+		/// <param name="ChangeNumber">The change number to sync. May be -1, for latest.</param>
+		/// <param name="Filters">List of filters to apply to the workspace. Each entry should be a path relative to the stream root, with an optional '-'prefix.</param>
+		private void UpdateClientHaveTable(PerforceConnection Client, string ClientName, int ChangeNumber, List<string> Filters)
+		{
+			using(LogStatusScope Scope = new LogStatusScope("Updating have table..."))
+			{
+				Stopwatch Timer = Stopwatch.StartNew();
+
+				// Sync an initial set of files. Either start with a full workspace and remove files, or start with nothing and add files.
+				if(Filters.Count == 0 || Filters[0].StartsWith("-"))
+				{
+					Client.SyncQuiet(SyncOptions.KeepWorkspaceFiles, -1, String.Format("//{0}/...@{1}", ClientName, ChangeNumber)).RequireSuccess();
+				}
+				else
+				{
+					Client.SyncQuiet(SyncOptions.KeepWorkspaceFiles, -1, String.Format("//{0}/...#0", ClientName)).RequireSuccess();
+				}
+
+				// Update with the contents of each filter
+				foreach(string Filter in Filters)
+				{
+					string SyncPath;
+					if(Filter.StartsWith("-"))
+					{
+						SyncPath = String.Format("//{0}/{1}#0", ClientName, RemoveLeadingSlash(Filter.Substring(1)));
+					}
+					else
+					{
+						SyncPath = String.Format("//{0}/{1}@{2}", ClientName, RemoveLeadingSlash(Filter), ChangeNumber);
+					}
+					Client.SyncQuiet(SyncOptions.KeepWorkspaceFiles, -1, SyncPath);
+				}
+
+				Scope.SetProgress("({0:0.0}s)", Timer.Elapsed.TotalSeconds);
+			}
+		}
+
+		/// <summary>
+		/// Get the contents of the client, as synced.
+		/// </summary>
+		/// <param name="Client">The client connection</param>
+		/// <param name="ClientName">Name of the current client</param>
+		/// <param name="ChangeNumber">The change number being synced. This must be specified in order to get the digest at the correct revision.</param>
+		/// <param name="bFakeSync">Whether this is for a fake sync. Poisons the file type to ensure that the cache is not corrupted.</param>
+		private StreamDirectoryInfo FindClientContents(PerforceConnection Client, string ClientName, int ChangeNumber, bool bFakeSync)
+		{
+			StreamDirectoryInfo Contents = new StreamDirectoryInfo();
+			using(LogStatusScope Scope = new LogStatusScope("Fetching metadata..."))
+			{
+				Stopwatch Timer = Stopwatch.StartNew();
+
+				// Find all the files in the branch
+				PerforceResponseList<FStatRecord> Response = Client.FStat(FStatOptions.OnlyHave | FStatOptions.IncludeFileSizes | FStatOptions.ClientFileInPerforceSyntax | FStatOptions.ShortenOutput, string.Format("//{0}/...@{1}", ClientName, ChangeNumber));
+				Response.RequireSuccess();
+
+				// Get the expected prefix for all files in client syntax
+				string ClientPrefix = string.Format("//{0}/", ClientName);
+
+				// Create the workspace, and add records for all the files. Exclude deleted files with digest = null.
+				foreach(FStatRecord Record in Response.Data)
+				{
+					if(!Record.ClientFile.StartsWith(ClientPrefix))
+					{
+						throw new InvalidDataException(String.Format("Client path returned by Perforce ('{0}') does not begin with client name ('{1}')", Record.ClientFile, ClientPrefix));
+					}
+					if(Record.Digest != null)
+					{
+						StreamDirectoryInfo LastStreamDirectory = Contents;
+					
+						string[] Fragments = PerforceUtils.UnescapePath(Record.ClientFile.Substring(ClientPrefix.Length)).Split('/');
+						for(int Idx = 0; Idx < Fragments.Length - 1; Idx++)
+						{
+							StreamDirectoryInfo NextStreamDirectory;
+							if(!LastStreamDirectory.NameToSubDirectory.TryGetValue(Fragments[Idx], out NextStreamDirectory))
+							{
+								NextStreamDirectory = new StreamDirectoryInfo(Fragments[Idx], LastStreamDirectory);
+								LastStreamDirectory.NameToSubDirectory.Add(Fragments[Idx], NextStreamDirectory);
+							}
+							LastStreamDirectory = NextStreamDirectory;
+						}
+
+						byte[] Digest = StringUtils.ParseHexString(Record.Digest);
+						FileContentId ContentId = new FileContentId(Digest, Record.HeadType + (bFakeSync? "+fake" : ""));
+
+						string DepotFileAndRevision = String.Format("{0}#{1}", Record.DepotFile, Record.HaveRevision);
+						LastStreamDirectory.NameToFile.Add(Fragments[Fragments.Length - 1], new StreamFileInfo(Fragments[Fragments.Length - 1], Record.FileSize, ContentId, LastStreamDirectory, DepotFileAndRevision));
+					}
+				}
+
+				Scope.SetProgress("({0:0.0}s)", Timer.Elapsed.TotalSeconds);
+			}
+			return Contents;
+		}
+
+		/// <summary>
+		/// Loads the contents of a client from disk
+		/// </summary>
+		/// <param name="CacheFile">The cache file to read from</param>
+		/// <returns>Contents of the workspace</returns>
+		StreamDirectoryInfo LoadClientContents(FileReference CacheFile)
+		{
+			StreamDirectoryInfo Contents;
+			using(LogStatusScope Scope = new LogStatusScope("Reading cached metadata from {0}...", CacheFile))
+			{
+				Stopwatch Timer = Stopwatch.StartNew();
+				Contents = StreamDirectoryInfo.Load(CacheFile);
+				Scope.SetProgress("({0:0.0}s)", Timer.Elapsed.TotalSeconds);
+			}
+			return Contents;
+		}
+
+		/// <summary>
+		/// Finds the contents of a workspace, and saves it to disk
+		/// </summary>
+		/// <param name="Client">The client connection</param>
+		/// <param name="ClientName">Name of the current client</param>
+		/// <param name="ChangeNumber">The change number being synced. This must be specified in order to get the digest at the correct revision.</param>
+		/// <param name="bFakeSync">Whether this is for a fake sync. Poisons the file type to ensure that the cache is not corrupted.</param>
+		/// <param name="CacheFile">Location of the file to save the cached contents</param>
+		/// <returns>Contents of the workspace</returns>
+		private StreamDirectoryInfo FindAndSaveClientContents(PerforceConnection Client, string ClientName, int ChangeNumber, bool bFakeSync, FileReference CacheFile)
+		{
+			StreamDirectoryInfo Contents = FindClientContents(Client, ClientName, ChangeNumber, bFakeSync);
+			using(LogStatusScope Scope = new LogStatusScope("Saving metadata to {0}...", CacheFile))
+			{
+				Stopwatch Timer = Stopwatch.StartNew();
+
+				// Handle the case where two machines may try to write to the cache file at once by writing to a temporary file
+				FileReference TempCacheFile = new FileReference(String.Format("{0}.{1}", CacheFile, Guid.NewGuid()));
+				Contents.Save(TempCacheFile);
+
+				// Try to move it into place
+				try
+				{
+					FileReference.Move(TempCacheFile, CacheFile);
+				}
+				catch(IOException)
+				{
+					if(!FileReference.Exists(CacheFile))
+					{
+						throw;
+					}
+					FileReference.Delete(TempCacheFile);
+				}
+
+				Scope.SetProgress("({0:0.0}s)", Timer.Elapsed.TotalSeconds);
+			}
+			return Contents;
+		}
+
+		/// <summary>
+		/// Remove files from the workspace
+		/// </summary>
+		/// <param name="Contents">Contents of the target stream</param>
+		private void RemoveFilesFromWorkspace(StreamDirectoryInfo Contents)
+		{
+			// Make sure the repair flag is clear before we start
+			RunOptionalRepair();
+
+			// Figure out what to remove
+			WorkspaceTransactionRemove Transaction;
+			using(LogStatusScope Scope = new LogStatusScope("Gathering files to remove..."))
+			{
+				Stopwatch Timer = Stopwatch.StartNew();
+
+				Transaction = new WorkspaceTransactionRemove(Workspace, Contents, ContentIdToTrackedFile);
+
+				Scope.SetProgress("({0:0.0}s)", Timer.Elapsed.TotalSeconds);
+			}
+
+			// Move files into the cache
+			KeyValuePair<FileContentId, WorkspaceFileInfo>[] FilesToMove = Transaction.FilesToMove.ToArray();
+			if(FilesToMove.Length > 0)
+			{
+				using(LogStatusScope Scope = new LogStatusScope("Moving {0} {1} to cache...", FilesToMove.Length, (FilesToMove.Length == 1)? "file" : "files"))
+				{
+					Stopwatch Timer = Stopwatch.StartNew();
+
+					// Add any new files to the cache
+					List<KeyValuePair<FileReference, FileReference>> SourceAndTargetFiles = new List<KeyValuePair<FileReference, FileReference>>();
+					foreach(KeyValuePair<FileContentId, WorkspaceFileInfo> FileToMove in FilesToMove)
+					{
+						ulong CacheId = GetUniqueCacheId(FileToMove.Key);
+						TrackedFileInfo NewTrackingInfo = new TrackedFileInfo(CacheDir, FileToMove.Key, CacheId, FileToMove.Value.Length, FileToMove.Value.LastModifiedTicks, FileToMove.Value.bReadOnly, NextSequenceNumber);
+						ContentIdToTrackedFile.Add(FileToMove.Key, NewTrackingInfo);
+						SourceAndTargetFiles.Add(new KeyValuePair<FileReference, FileReference>(FileToMove.Value.GetLocation(), NewTrackingInfo.GetLocation()));
+					}
+					NextSequenceNumber++;
+
+					// Save the current state of the repository as dirty. If we're interrupted, we will have two places to check for each file (the cache and workspace).
+					Save(RepositoryState.Dirty);
+
+					// Execute all the moves and deletes
+					Parallel.ForEach(SourceAndTargetFiles, SourceAndTargetFile => FileUtils.ForceMoveFile(SourceAndTargetFile.Key, SourceAndTargetFile.Value));
+
+					Scope.SetProgress("({0:0.0}s)", Timer.Elapsed.TotalSeconds);
+				}
+			}
+
+			// Remove files which are no longer needed
+			WorkspaceFileInfo[] FilesToDelete = Transaction.FilesToDelete.ToArray();
+			if(FilesToDelete.Length > 0)
+			{
+				using(LogStatusScope Scope = new LogStatusScope("Deleting {0} {1}...", FilesToDelete.Length, (FilesToDelete.Length == 1)? "file" : "files"))
+				{
+					Stopwatch Timer = Stopwatch.StartNew();
+
+					Parallel.ForEach(FilesToDelete, FileToDelete => RemoveFile(FileToDelete));
+
+					Scope.SetProgress("({0:0.0}s)", Timer.Elapsed.TotalSeconds);
+				}
+			}
+
+			// Remove directories which are no longer needed
+			WorkspaceDirectoryInfo[] DirectoriesToDelete = Transaction.DirectoriesToDelete.ToArray();
+			if(DirectoriesToDelete.Length > 0)
+			{
+				using(LogStatusScope Scope = new LogStatusScope("Deleting {0} {1}...", DirectoriesToDelete.Length, (DirectoriesToDelete.Length == 1)? "directory" : "directories"))
+				{
+					Stopwatch Timer = Stopwatch.StartNew();
+
+					foreach(string DirectoryToDelete in DirectoriesToDelete.Select(x => x.GetFullName()).OrderByDescending(x => x.Length))
+					{
+						RemoveDirectory(DirectoryToDelete);
+					}
+
+					Scope.SetProgress("({0:0.0}s)", Timer.Elapsed.TotalSeconds);
+				}
+			}
+
+			// Update the workspace and save the new state
+			Workspace = Transaction.NewWorkspaceRootDir;
+			Save(RepositoryState.Clean);
+		}
+
+		/// <summary>
+		/// Helper function to delete a file from the workspace, and output any failure as a warning.
+		/// </summary>
+		/// <param name="FileToDelete">The file to be deleted</param>
+		void RemoveFile(WorkspaceFileInfo FileToDelete)
+		{
+			try
+			{
+				FileUtils.ForceDeleteFile(FileToDelete.GetLocation());
+			}
+			catch(Exception Ex)
+			{
+				Log.TraceWarning("warning: Unable to delete file {0}.", FileToDelete.GetFullName());
+				Log.TraceVerbose(ExceptionUtils.FormatExceptionDetails(Ex));
+				bRequiresRepair = true;
+			}
+		}
+
+		/// <summary>
+		/// Helper function to delete a directory from the workspace, and output any failure as a warning.
+		/// </summary>
+		/// <param name="DirectoryToDelete">The directory to be deleted</param>
+		void RemoveDirectory(string DirectoryToDelete)
+		{
+			try
+			{
+				Directory.Delete(DirectoryToDelete, false);
+			}
+			catch(Exception Ex)
+			{
+				Log.TraceWarning("warning: Unable to delete directory {0}", DirectoryToDelete);
+				Log.TraceVerbose(ExceptionUtils.FormatExceptionDetails(Ex));
+				bRequiresRepair = true;
+			}
+		}
+
+		/// <summary>
+		/// Update the workspace to match the given stream, syncing files and moving to/from the cache as necessary.
+		/// </summary>
+		/// <param name="Client">The client connection</param>
+		/// <param name="Stream">Contents of the stream</param>
+		/// <param name="bFakeSync">Whether to simulate the sync operation, rather than actually syncing files</param>
+		private void AddFilesToWorkspace(PerforceConnection Client, StreamDirectoryInfo Stream, bool bFakeSync)
+		{
+			// Make sure the repair flag is reset
+			RunOptionalRepair();
+
+			// Figure out what we need to do
+			WorkspaceTransactionAdd Transaction;
+			using(LogStatusScope Status = new LogStatusScope("Gathering files to add..."))
+			{
+				Stopwatch Timer = Stopwatch.StartNew();
+
+				Transaction = new WorkspaceTransactionAdd(Workspace, Stream, ContentIdToTrackedFile);
+				Workspace = Transaction.NewWorkspaceRootDir;
+				Save(RepositoryState.Dirty);
+
+				Status.SetProgress("({0:0.0}s)", Timer.Elapsed.TotalSeconds);
+			}
+
+			// Swap files in and out of the cache
+			WorkspaceFileToMove[] FilesToMove = Transaction.FilesToMove.Values.ToArray();
+			if(FilesToMove.Length > 0)
+			{
+				using(LogStatusScope Status = new LogStatusScope("Moving {0} {1} from cache...", FilesToMove.Length, (FilesToMove.Length == 1)? "file" : "files"))
+				{
+					Stopwatch Timer = Stopwatch.StartNew();
+					Parallel.ForEach(FilesToMove, FileToMove => MoveFileFromCache(FileToMove, Transaction.FilesToSync));
+					ContentIdToTrackedFile = ContentIdToTrackedFile.Where(x => !Transaction.FilesToMove.ContainsKey(x.Value)).ToDictionary(x => x.Key, x => x.Value);
+					Status.SetProgress("({0:0.0}s)", Timer.Elapsed.TotalSeconds);
+				}
+			}
+
+			// Swap files in and out of the cache
+			WorkspaceFileToCopy[] FilesToCopy = Transaction.FilesToCopy.ToArray();
+			if(FilesToCopy.Length > 0)
+			{
+				using(LogStatusScope Status = new LogStatusScope("Copying {0} {1} within workspace...", FilesToCopy.Length, (FilesToCopy.Length == 1)? "file" : "files"))
+				{
+					Stopwatch Timer = Stopwatch.StartNew();
+					Parallel.ForEach(FilesToCopy, FileToCopy => CopyFileWithinWorkspace(FileToCopy, Transaction.FilesToSync));
+					Status.SetProgress("({0:0.0}s)", Timer.Elapsed.TotalSeconds);
+				}
+			}
+
+			// Find all the files we want to sync
+			WorkspaceFileToSync[] FilesToSync = Transaction.FilesToSync.ToArray();
+			if(FilesToSync.Length > 0)
+			{
+				// Sync all the files
+				long TotalSize = FilesToSync.Sum(x => x.StreamFile.Length);
+				using(LogStatusScope Status = new LogStatusScope("Syncing {0} {1}...", FilesToSync.Length, (FilesToSync.Length == 1)? "file" : "files"))
+				{
+					Stopwatch Timer = Stopwatch.StartNew();
+
+					long SyncedSize = 0;
+					for(int EndIdx = 0; EndIdx < FilesToSync.Length; )
+					{
+						int BeginIdx = EndIdx;
+
+						// Update the progress
+						Status.SetProgress("{0:n1}% ({1:n1}mb/{2:n1}mb)", SyncedSize * 100.0 / TotalSize, SyncedSize / (1024.0 * 1024.0), TotalSize / (1024.0 * 1024.0));
+
+						// Figure out the next batch of files to sync
+						long BatchSize = 0;
+						for(; EndIdx < FilesToSync.Length && BatchSize < 256 * 1024 * 1024; EndIdx++)
+						{
+							BatchSize += FilesToSync[EndIdx].StreamFile.Length;
+						}
+
+						// Print the list of files we're syncing to the log
+						Log.TraceLog("Batch {0}-{1}, {2:n1}mb:", BeginIdx, EndIdx, BatchSize / (1024.0 * 1024.0));
+						for(int Idx = BeginIdx; Idx < EndIdx; Idx++)
+						{
+							Log.TraceLog("  {0,-128} [{1,-48}]", FilesToSync[Idx].StreamFile.DepotFileAndRevision, FilesToSync[Idx].StreamFile.ContentId);
+						}
+
+						// Sync the files
+						if(bFakeSync)
+						{
+							for(int Idx = BeginIdx; Idx < EndIdx; Idx++)
+							{
+								FileReference LocalFile = FilesToSync[Idx].WorkspaceFile.GetLocation();
+								DirectoryReference.CreateDirectory(LocalFile.Directory);
+								FileReference.WriteAllBytes(LocalFile, new byte[0]);
+							}
+						}
+						else
+						{
+							FileReference SyncFileName = FileReference.Combine(BaseDir, "SyncList.txt");
+
+							using(StreamWriter Writer = new StreamWriter(SyncFileName.FullName))
+							{
+								for(int Idx = BeginIdx; Idx < EndIdx; Idx++)
+								{
+									Writer.WriteLine(FilesToSync[Idx].StreamFile.DepotFileAndRevision);
+								}
+							}
+
+							PerforceConnection TempPerforce = new PerforceConnection(String.Format("{0} -x\"{1}\"", Client.GlobalOptions, SyncFileName));
+							TempPerforce.Sync(SyncOptions.Force | SyncOptions.FullDepotSyntax, -1).RequireSuccess();
+						}
+
+						// Update metadata for the current state
+						Parallel.For(BeginIdx, EndIdx, Idx => FilesToSync[Idx].WorkspaceFile.UpdateMetadata());
+
+						// Save the current state
+						Save(RepositoryState.Dirty);
+
+						// Update the status
+						SyncedSize += BatchSize;
+					}
+
+					Status.SetProgress("100% ({0:n1}mb) ({1:0.0}s)", TotalSize / (1024 * 1024.0), Timer.Elapsed.TotalSeconds);
+				}
+			}
+
+			// Save the clean state
+			Workspace = Transaction.NewWorkspaceRootDir;
+			Save(RepositoryState.Clean);
+		}
+
+		/// <summary>
+		/// Helper function to move a file from the cache into the workspace. If it fails, adds the file to a list to be synced.
+		/// </summary>
+		/// <param name="FileToMove">Information about the file to move</param>
+		/// <param name="FilesToSync">List of files to be synced. If the move fails, the file will be added to this list of files to sync.</param>
+		void MoveFileFromCache(WorkspaceFileToMove FileToMove, ConcurrentBag<WorkspaceFileToSync> FilesToSync)
+		{
+			try
+			{
+				FileReference.Move(FileToMove.TrackedFile.GetLocation(), FileToMove.WorkspaceFile.GetLocation());
+			}
+			catch(Exception Ex)
+			{
+				Log.TraceWarning("warning: Unable to move {0} from cache to {1}. Syncing instead.", FileToMove.TrackedFile.GetLocation(), FileToMove.WorkspaceFile.GetLocation());
+				Log.TraceVerbose(ExceptionUtils.FormatExceptionDetails(Ex));
+				FilesToSync.Add(new WorkspaceFileToSync(FileToMove.StreamFile, FileToMove.WorkspaceFile));
+				bRequiresRepair = true;
+			}
+		}
+
+		/// <summary>
+		/// Helper function to copy a file within the workspace. If it fails, adds the file to a list to be synced.
+		/// </summary>
+		/// <param name="FileToCopy">Information about the file to move</param>
+		/// <param name="FilesToSync">List of files to be synced. If the move fails, the file will be added to this list of files to sync.</param>
+		void CopyFileWithinWorkspace(WorkspaceFileToCopy FileToCopy, ConcurrentBag<WorkspaceFileToSync> FilesToSync)
+		{
+			try
+			{
+				FileReference.Copy(FileToCopy.SourceWorkspaceFile.GetLocation(), FileToCopy.TargetWorkspaceFile.GetLocation());
+				FileToCopy.TargetWorkspaceFile.UpdateMetadata();
+			}
+			catch(Exception Ex)
+			{
+				Log.TraceWarning("warning: Unable to copy {0} to {1}. Syncing instead.", FileToCopy.SourceWorkspaceFile.GetLocation(), FileToCopy.TargetWorkspaceFile.GetLocation());
+				Log.TraceVerbose(ExceptionUtils.FormatExceptionDetails(Ex));
+				FilesToSync.Add(new WorkspaceFileToSync(FileToCopy.StreamFile, FileToCopy.TargetWorkspaceFile));
+				bRequiresRepair = true;
+			}
+		}
+
+		void RemoveTrackedFile(TrackedFileInfo TrackedFile)
+		{
+			ContentIdToTrackedFile.Remove(TrackedFile.ContentId);
+			CacheEntries.Remove(TrackedFile.CacheId);
+			FileUtils.ForceDeleteFile(TrackedFile.GetLocation());
+		}
+
+		void CreateCacheHierarchy()
+		{
+			for(int IdxA = 0; IdxA < 16; IdxA++)
+			{
+				DirectoryReference DirA = DirectoryReference.Combine(CacheDir, String.Format("{0:X}", IdxA));
+				DirectoryReference.CreateDirectory(DirA);
+				for(int IdxB = 0; IdxB < 16; IdxB++)
+				{
+					DirectoryReference DirB = DirectoryReference.Combine(DirA, String.Format("{0:X}", IdxB));
+					DirectoryReference.CreateDirectory(DirB);
+					for(int IdxC = 0; IdxC < 16; IdxC++)
+					{
+						DirectoryReference DirC = DirectoryReference.Combine(DirB, String.Format("{0:X}", IdxC));
+						DirectoryReference.CreateDirectory(DirC);
+					}
+				}
+			}
+		}
+
+		/// <summary>
+		/// Determines a unique cache id for a file content id
+		/// </summary>
+		/// <param name="ContentId">File content id to get a unique id for</param>
+		/// <returns>The unique cache id</returns>
+		ulong GetUniqueCacheId(FileContentId ContentId)
+		{
+			// Initialize the cache id to the top 16 bytes of the digest, then increment it until we find a unique id
+			ulong CacheId = 0;
+			for(int Idx = 0; Idx < 8; Idx++)
+			{
+				CacheId = (CacheId << 8) | ContentId.Digest[Idx];
+			}
+			while(!CacheEntries.Add(CacheId))
+			{
+				CacheId++;
+			}
+			return CacheId;
+		}
+
+		/// <summary>
+		/// Sanitizes a string for use in a Perforce client name
+		/// </summary>
+		/// <param name="Text">Text to sanitize</param>
+		/// <returns>Sanitized text</returns>
+		static string Sanitize(string Text)
+		{
+			StringBuilder Result = new StringBuilder();
+			for(int Idx = 0; Idx < Text.Length; Idx++)
+			{
+				if((Text[Idx] >= '0' && Text[Idx] <= '9') || (Text[Idx] >= 'a' && Text[Idx] <= 'z') || (Text[Idx] >= 'A' && Text[Idx] <= 'Z') || Text[Idx] == '-' || Text[Idx] == '_' || Text[Idx] == '.')
+				{
+					Result.Append(Text[Idx]);
+				}
+			}
+			return Result.ToString();
+		}
+
+		/// <summary>
+		/// Removes the leading slash from a path
+		/// </summary>
+		/// <param name="Path">The path to remove a slash from</param>
+		/// <returns>The path without a leading slash</returns>
+		static string RemoveLeadingSlash(string Path)
+		{
+			if(Path.Length > 0 && Path[0] == '/')
+			{
+				return Path.Substring(1);
+			}
+			else
+			{
+				return Path;
+			}
+		}
+
+		/// <summary>
+		/// Gets the path to a backup file used while a new file is being written out
+		/// </summary>
+		/// <param name="TargetFile">The file being written to</param>
+		/// <returns>The path to a backup file</returns>
+		private static FileReference GetBackupFile(FileReference TargetFile)
+		{
+			return new FileReference(TargetFile.FullName + ".transaction");
+		}
+
+		/// <summary>
+		/// Begins a write transaction on the given file. Assumes only one process will be reading/writing at a time, but the operation can be interrupted.
+		/// </summary>
+		/// <param name="TargetFile">The file being written to</param>
+		public static void BeginTransaction(FileReference TargetFile)
+		{
+			FileReference TransactionFile = GetBackupFile(TargetFile);
+			if (FileReference.Exists(TargetFile))
+			{
+				FileUtils.ForceMoveFile(TargetFile, TransactionFile);
+			}
+			else if (FileReference.Exists(TransactionFile))
+			{
+				FileUtils.ForceDeleteFile(TransactionFile);
+			}
+		}
+
+		/// <summary>
+		/// Mark a transaction on the given file as complete, and removes the backup file.
+		/// </summary>
+		/// <param name="TargetFile">The file being written to</param>
+		public static void CompleteTransaction(FileReference TargetFile)
+		{
+			FileReference TransactionFile = GetBackupFile(TargetFile);
+			FileUtils.ForceDeleteFile(TransactionFile);
+		}
+
+		/// <summary>
+		/// Restores the backup for a target file, if it exists. This allows recovery from an incomplete transaction.
+		/// </summary>
+		/// <param name="TargetFile">The file being written to</param>
+		public static void RestoreBackup(FileReference TargetFile)
+		{
+			FileReference TransactionFile = GetBackupFile(TargetFile);
+			if (FileReference.Exists(TransactionFile))
+			{
+				FileUtils.ForceMoveFile(TransactionFile, TargetFile);
+			}
+		}
+
+		#endregion
+	}
+}
diff --git a/Engine/Source/Programs/BuildAgent/Workspace/Common/StreamInfo.cs b/Engine/Source/Programs/BuildAgent/Workspace/Common/StreamInfo.cs
new file mode 100644
--- /dev/null
+++ b/Engine/Source/Programs/BuildAgent/Workspace/Common/StreamInfo.cs
@@ -0,0 +1,312 @@
+// Copyright 1998-2019 Epic Games, Inc. All Rights Reserved.
+
+using System;
+using System.Collections.Generic;
+using System.Diagnostics;
+using System.IO;
+using System.IO.Compression;
+using System.Linq;
+using System.Text;
+using Tools.DotNETCommon;
+using Tools.DotNETCommon.Perforce;
+
+namespace BuildAgent.Workspace.Common
+{
+	/// <summary>
+	/// Represents a file within a stream
+	/// </summary>
+	class StreamFileInfo
+	{
+		/// <summary>
+		/// Name of this file
+		/// </summary>
+		public readonly string Name;
+
+		/// <summary>
+		/// Length of the file, as reported by the server (actual size on disk may be different due to workspace options).
+		/// </summary>
+		public readonly long Length;
+
+		/// <summary>
+		/// Content id for this file
+		/// </summary>
+		public readonly FileContentId ContentId;
+
+		/// <summary>
+		/// The parent directory
+		/// </summary>
+		public readonly StreamDirectoryInfo Directory;
+
+		/// <summary>
+		/// The depot file and revision that need to be synced for this file
+		/// </summary>
+		public readonly string DepotFileAndRevision;
+
+		/// <summary>
+		/// Constructor
+		/// </summary>
+		/// <param name="Name">Name of this file</param>
+		/// <param name="Length">Length of the file on the server</param>
+		/// <param name="ContentId">Content id for this file</param>
+		/// <param name="Directory">The parent directory</param>
+		/// <param name="DepotFileAndRevision">The depot file and revision that need to be synced for this file</param>
+		public StreamFileInfo(string Name, long Length, FileContentId ContentId, StreamDirectoryInfo Directory, string DepotFileAndRevision)
+		{
+			this.Name = Name;
+			this.Length = Length;
+			this.ContentId = ContentId;
+			this.Directory = Directory;
+			this.DepotFileAndRevision = DepotFileAndRevision;
+		}
+
+		/// <summary>
+		/// Constructor for reading a file info from disk
+		/// </summary>
+		/// <param name="Reader">Binary reader to read data from</param>
+		/// <param name="Directory">Parent directory</param>
+		public StreamFileInfo(BinaryReader Reader, StreamDirectoryInfo Directory)
+		{
+			this.Name = Reader.ReadString();
+			this.Length = Reader.ReadInt64();
+			this.ContentId = new FileContentId(Reader);
+			this.Directory = Directory;
+			this.DepotFileAndRevision = Reader.ReadString();
+		}
+
+		/// <summary>
+		/// Save the file info to disk
+		/// </summary>
+		/// <param name="Writer">Writer to output to</param>
+		public void Write(BinaryWriter Writer)
+		{
+			Writer.Write(Name);
+			Writer.Write(Length);
+			ContentId.Write(Writer);
+			Writer.Write(DepotFileAndRevision);
+		}
+
+		/// <summary>
+		/// Get the path to this file relative to the root of the stream
+		/// </summary>
+		/// <returns>Relative path to the file</returns>
+		public string GetRelativePath()
+		{
+			StringBuilder Builder = new StringBuilder();
+			Directory.AppendPath(Builder);
+			Builder.Append(Name);
+			return Builder.ToString();
+		}
+
+		/// <summary>
+		/// Format the path to the file for the debugger
+		/// </summary>
+		/// <returns>Path to the file</returns>
+		public override string ToString()
+		{
+			return GetRelativePath();
+		}
+	}
+
+	/// <summary>
+	/// Information about a directory within a stream
+	/// </summary>
+	class StreamDirectoryInfo
+	{
+		/// <summary>
+		/// The current signature for saved directory objects
+		/// </summary>
+		static readonly byte[] CurrentSignature = { (byte)'W', (byte)'S', (byte)'D', 1 };
+
+		/// <summary>
+		/// The directory name
+		/// </summary>
+		public readonly string Name;
+
+		/// <summary>
+		/// The parent directory
+		/// </summary>
+		public readonly StreamDirectoryInfo ParentDirectory;
+
+		/// <summary>
+		/// Map of name to file within the directory
+		/// </summary>
+		public Dictionary<string, StreamFileInfo> NameToFile = new Dictionary<string, StreamFileInfo>(StringComparer.Ordinal);
+
+		/// <summary>
+		/// Map of name to subdirectory
+		/// </summary>
+		public Dictionary<string, StreamDirectoryInfo> NameToSubDirectory = new Dictionary<string, StreamDirectoryInfo>(FileUtils.PlatformPathComparer);
+
+		/// <summary>
+		/// Constructor
+		/// </summary>
+		public StreamDirectoryInfo()
+			: this("", null)
+		{
+		}
+
+		/// <summary>
+		/// Constructor
+		/// </summary>
+		/// <param name="Name">Name of the directory</param>
+		/// <param name="ParentDirectory">Parent directory</param>
+		public StreamDirectoryInfo(string Name, StreamDirectoryInfo ParentDirectory)
+		{
+			this.Name = Name;
+			this.ParentDirectory = ParentDirectory;
+		}
+
+		/// <summary>
+		/// Constructor for reading from disk
+		/// </summary>
+		/// <param name="Reader">Reader to read data from</param>
+		/// <param name="ParentDirectory">The parent directory to initialize this directory with</param>
+		public StreamDirectoryInfo(BinaryReader Reader, StreamDirectoryInfo ParentDirectory)
+		{
+			this.Name = Reader.ReadString();
+			this.ParentDirectory = ParentDirectory;
+
+			int NumFiles = Reader.ReadInt32();
+			for(int Idx = 0; Idx < NumFiles; Idx++)
+			{
+				StreamFileInfo File = new StreamFileInfo(Reader, this);
+				NameToFile.Add(File.Name, File);
+			}
+
+			int NumDirectories = Reader.ReadInt32();
+			for(int Idx = 0; Idx < NumDirectories; Idx++)
+			{
+				StreamDirectoryInfo SubDirectory = new StreamDirectoryInfo(Reader, this);
+				NameToSubDirectory.Add(SubDirectory.Name, SubDirectory);
+			}
+		}
+
+		/// <summary>
+		/// Writes the contents of this stream to disk
+		/// </summary>
+		/// <param name="Writer">Writer to serialize to</param>
+		public void Write(BinaryWriter Writer)
+		{
+			Writer.Write(Name);
+
+			Writer.Write(NameToFile.Count);
+			foreach(StreamFileInfo File in NameToFile.Values)
+			{
+				File.Write(Writer);
+			}
+
+			Writer.Write(NameToSubDirectory.Count);
+			foreach(StreamDirectoryInfo SubDirectory in NameToSubDirectory.Values)
+			{
+				SubDirectory.Write(Writer);
+			}
+		}
+
+		/// <summary>
+		/// Load a stream directory from a file on disk
+		/// </summary>
+		/// <param name="InputFile">File to read from</param>
+		/// <returns>New StreamDirectoryInfo object</returns>
+		public static StreamDirectoryInfo Load(FileReference InputFile)
+		{
+			using(FileStream Stream = File.Open(InputFile.FullName, FileMode.Open, FileAccess.Read, FileShare.Read))
+			{
+				byte[] Signature = new byte[CurrentSignature.Length];
+				if(Stream.Read(Signature, 0, CurrentSignature.Length) != CurrentSignature.Length)
+				{
+					throw new InvalidDataException(String.Format("Unable to read signature bytes from {0}", InputFile));
+				}
+				if(!Enumerable.SequenceEqual(Signature, CurrentSignature))
+				{
+					throw new InvalidDataException(String.Format("Cached stream contents at {0} has incorrect signature", InputFile));
+				}
+
+				using(GZipStream CompressedStream = new GZipStream(Stream, CompressionMode.Decompress))
+				{
+					using(BinaryReader Reader = new BinaryReader(CompressedStream, Encoding.UTF8, true))
+					{
+						return new StreamDirectoryInfo(Reader, null);
+					}
+				}
+			}
+		}
+
+		/// <summary>
+		/// Saves the contents of this object to disk
+		/// </summary>
+		/// <param name="OutputFile">The output file to write to</param>
+		public void Save(FileReference OutputFile)
+		{
+			using(FileStream Stream = File.Open(OutputFile.FullName, FileMode.CreateNew, FileAccess.Write, FileShare.Read))
+			{
+				Stream.Write(CurrentSignature, 0, CurrentSignature.Length);
+
+				using(GZipStream CompressedStream = new GZipStream(Stream, CompressionMode.Compress))
+				{
+					using(BinaryWriter Writer = new BinaryWriter(CompressedStream, Encoding.UTF8, true))
+					{
+						Write(Writer);
+					}
+				}
+			}
+		}
+
+		/// <summary>
+		/// Get all the files in this directory
+		/// </summary>
+		/// <returns>List of files</returns>
+		public List<StreamFileInfo> GetFiles()
+		{
+			List<StreamFileInfo> Files = new List<StreamFileInfo>();
+			AppendFiles(Files);
+			return Files;
+		}
+
+		/// <summary>
+		/// Append the contents of this directory and subdirectories to a list
+		/// </summary>
+		/// <param name="Files">List to append to</param>
+		public void AppendFiles(List<StreamFileInfo> Files)
+		{
+			foreach(StreamDirectoryInfo SubDirectory in NameToSubDirectory.Values)
+			{
+				SubDirectory.AppendFiles(Files);
+			}
+			Files.AddRange(NameToFile.Values);
+		}
+
+		/// <summary>
+		/// Gets the relative path to this directory, with a trailing slash
+		/// </summary>
+		/// <returns>Relative path to this directory</returns>
+		public string GetRelativePath()
+		{
+			StringBuilder Builder = new StringBuilder();
+			AppendPath(Builder);
+			return Builder.ToString();
+		}
+
+		/// <summary>
+		/// Append the relative path to this directory to the given string builder
+		/// </summary>
+		/// <param name="Builder">String builder to append to</param>
+		public void AppendPath(StringBuilder Builder)
+		{
+			if(ParentDirectory != null)
+			{
+				ParentDirectory.AppendPath(Builder);
+				Builder.Append(Name);
+			}
+			Builder.Append('/');
+		}
+
+		/// <summary>
+		/// Format the path to the directory for debugging
+		/// </summary>
+		/// <returns>Path to the directory</returns>
+		public override string ToString()
+		{
+			return GetRelativePath();
+		}
+	}
+}
diff --git a/Engine/Source/Programs/BuildAgent/Workspace/Common/TrackedFileInfo.cs b/Engine/Source/Programs/BuildAgent/Workspace/Common/TrackedFileInfo.cs
new file mode 100644
--- /dev/null
+++ b/Engine/Source/Programs/BuildAgent/Workspace/Common/TrackedFileInfo.cs
@@ -0,0 +1,104 @@
+// Copyright 1998-2019 Epic Games, Inc. All Rights Reserved.
+
+using System;
+using System.Collections.Generic;
+using System.Diagnostics;
+using System.IO;
+using System.Linq;
+using System.Text;
+using System.Threading;
+using System.Threading.Tasks;
+using Tools.DotNETCommon;
+
+namespace BuildAgent.Workspace.Common
+{
+	/// <summary>
+	/// Stores information about where a particular FileContentId has been staged, or its location in the cache
+	/// </summary>
+	[DebuggerDisplay("{ContentId}")]
+	class TrackedFileInfo : IBinarySerializable
+	{
+		public readonly DirectoryReference CacheDir;
+		public readonly FileContentId ContentId;
+		public readonly ulong CacheId;
+		public readonly long Length;
+		public readonly long LastModifiedTicks;
+		public readonly bool bReadOnly;
+		public readonly uint SequenceNumber;
+
+		public TrackedFileInfo(DirectoryReference CacheDir, FileContentId ContentId, ulong CacheId, long Length, long LastModifiedTicks, bool bReadOnly, uint SequenceNumber)
+		{
+			this.CacheDir = CacheDir;
+			this.ContentId = ContentId;
+			this.CacheId = CacheId;
+			this.Length = Length;
+			this.LastModifiedTicks = LastModifiedTicks;
+			this.bReadOnly = bReadOnly;
+			this.SequenceNumber = SequenceNumber;
+		}
+
+		public TrackedFileInfo(DirectoryReference CacheDir, BinaryReader Reader)
+		{
+			this.CacheDir = CacheDir;
+			this.ContentId = Reader.ReadObject<FileContentId>();
+			this.CacheId = Reader.ReadUInt64();
+			this.Length = Reader.ReadInt64();
+			this.LastModifiedTicks = Reader.ReadInt64();
+			this.bReadOnly = Reader.ReadBoolean();
+			this.SequenceNumber = Reader.ReadUInt32();
+		}
+
+		public void Write(BinaryWriter Writer)
+		{
+			Writer.Write(ContentId);
+			Writer.Write(CacheId);
+			Writer.Write(Length);
+			Writer.Write(LastModifiedTicks);
+			Writer.Write(bReadOnly);
+			Writer.Write(SequenceNumber);
+		}
+
+		public bool CheckIntegrity()
+		{
+			FileReference Location = GetLocation();
+
+			FileInfo Info = new FileInfo(Location.FullName);
+			if(!Info.Exists)
+			{
+				Log.TraceWarning("warning: {0} was missing from cache.", Location);
+				return false;
+			}
+			if(Info.Length != Length)
+			{
+				Log.TraceWarning("warning: {0} was {1:n} bytes; expected {2:n} bytes", Location, Info.Length, Length);
+				return false;
+			}
+			if(Info.LastWriteTimeUtc.Ticks != LastModifiedTicks)
+			{
+				Log.TraceWarning("warning: {0} was last modified at {1}; expected {2}", Location, Info.LastWriteTimeUtc, new DateTime(LastModifiedTicks, DateTimeKind.Utc));
+				return false;
+			}
+			if(Info.Attributes.HasFlag(FileAttributes.ReadOnly) != bReadOnly)
+			{
+				Log.TraceWarning("warning: {0} readonly flag is {1}; expected {2}", Info.Attributes.HasFlag(FileAttributes.ReadOnly), bReadOnly);
+				return false;
+			}
+
+			return true;
+		}
+
+		public FileReference GetLocation()
+		{
+			StringBuilder FullName = new StringBuilder(CacheDir.FullName);
+			FullName.Append(Path.DirectorySeparatorChar);
+			FullName.AppendFormat("{0:X}", (CacheId >> 60) & 0xf);
+			FullName.Append(Path.DirectorySeparatorChar);
+			FullName.AppendFormat("{0:X}", (CacheId >> 56) & 0xf);
+			FullName.Append(Path.DirectorySeparatorChar);
+			FullName.AppendFormat("{0:X}", (CacheId >> 62) & 0xf);
+			FullName.Append(Path.DirectorySeparatorChar);
+			FullName.AppendFormat("{0:X}", CacheId);
+			return new FileReference(FullName.ToString());
+		}
+	}
+}
diff --git a/Engine/Source/Programs/BuildAgent/Workspace/Common/WorkspaceInfo.cs b/Engine/Source/Programs/BuildAgent/Workspace/Common/WorkspaceInfo.cs
new file mode 100644
--- /dev/null
+++ b/Engine/Source/Programs/BuildAgent/Workspace/Common/WorkspaceInfo.cs
@@ -0,0 +1,334 @@
+// Copyright 1998-2019 Epic Games, Inc. All Rights Reserved.
+
+using System;
+using System.Collections.Concurrent;
+using System.Collections.Generic;
+using System.Diagnostics;
+using System.IO;
+using System.Linq;
+using System.Text;
+using System.Threading;
+using System.Threading.Tasks;
+using Tools.DotNETCommon;
+
+namespace BuildAgent.Workspace.Common
+{
+	/// <summary>
+	/// Stores information about a file that has been staged into a workspace
+	/// </summary>
+	class WorkspaceFileInfo
+	{
+		public readonly WorkspaceDirectoryInfo Directory;
+		public readonly string Name;
+		public long Length;
+		public long LastModifiedTicks;
+		public bool bReadOnly;
+		public readonly FileContentId ContentId;
+
+		public WorkspaceFileInfo(WorkspaceDirectoryInfo Directory, string Name, FileContentId ContentId)
+		{
+			this.Directory = Directory;
+			this.Name = Name;
+			this.ContentId = ContentId;
+		}
+
+		public WorkspaceFileInfo(WorkspaceDirectoryInfo Directory, string Name, FileInfo Info, FileContentId ContentId)
+			: this(Directory, Name, Info.Length, Info.LastWriteTimeUtc.Ticks, Info.Attributes.HasFlag(FileAttributes.ReadOnly), ContentId)
+		{
+		}
+
+		public WorkspaceFileInfo(WorkspaceDirectoryInfo Directory, string Name, long Length, long LastModifiedTicks, bool bReadOnly, FileContentId ContentId)
+		{
+			this.Directory = Directory;
+			this.Name = Name;
+			this.Length = Length;
+			this.LastModifiedTicks = LastModifiedTicks;
+			this.bReadOnly = bReadOnly;
+			this.ContentId = ContentId;
+		}
+
+		public WorkspaceFileInfo(WorkspaceDirectoryInfo Directory, BinaryReader Reader)
+		{
+			this.Directory = Directory;
+			this.Name = Reader.ReadString();
+			this.Length = Reader.ReadInt64();
+			this.LastModifiedTicks = Reader.ReadInt64();
+			this.bReadOnly = Reader.ReadBoolean();
+			this.ContentId = Reader.ReadObject<FileContentId>();
+		}
+
+		public void SetMetadata(long Length, long LastModifiedTicks, bool bReadOnly)
+		{
+			this.Length = Length;
+			this.LastModifiedTicks = LastModifiedTicks;
+			this.bReadOnly = bReadOnly;
+		}
+
+		public void UpdateMetadata()
+		{
+			FileInfo Info = new FileInfo(GetFullName());
+			if(Info.Exists)
+			{
+				Length = Info.Length;
+				LastModifiedTicks = Info.LastWriteTimeUtc.Ticks;
+				bReadOnly = Info.Attributes.HasFlag(FileAttributes.ReadOnly);
+			}
+		}
+
+		public bool MatchesAttributes(FileInfo Info)
+		{
+			return Length == Info.Length && LastModifiedTicks == Info.LastWriteTimeUtc.Ticks && (Info.Attributes.HasFlag(FileAttributes.ReadOnly) == bReadOnly);
+		}
+
+		public void Write(BinaryWriter Writer)
+		{
+			Writer.Write(Name);
+			Writer.Write(Length);
+			Writer.Write(LastModifiedTicks);
+			Writer.Write(bReadOnly);
+			Writer.Write(ContentId);
+		}
+
+		public string GetClientPath()
+		{
+			StringBuilder Builder = new StringBuilder();
+			Directory.AppendClientPath(Builder);
+			Builder.Append(Name);
+			return Builder.ToString();
+		}
+
+		public string GetFullName()
+		{
+			StringBuilder Builder = new StringBuilder();
+			Directory.AppendFullPath(Builder);
+			Builder.Append(Path.DirectorySeparatorChar);
+			Builder.Append(Name);
+			return Builder.ToString();
+		}
+
+		public FileReference GetLocation()
+		{
+			return new FileReference(GetFullName());
+		}
+
+		public override string ToString()
+		{
+			return GetFullName();
+		}
+	}
+
+	/// <summary>
+	/// 
+	/// </summary>
+	class WorkspaceDirectoryInfo
+	{
+		public readonly WorkspaceDirectoryInfo ParentDirectory;
+		public readonly string Name;
+		public Dictionary<string, WorkspaceFileInfo> NameToFile;
+		public Dictionary<string, WorkspaceDirectoryInfo> NameToSubDirectory;
+
+		public WorkspaceDirectoryInfo(DirectoryReference RootDir)
+			: this(null, RootDir.FullName)
+		{
+		}
+
+		public WorkspaceDirectoryInfo(WorkspaceDirectoryInfo ParentDirectory, string Name)
+		{
+			this.ParentDirectory = ParentDirectory;
+			this.Name = Name;
+			this.NameToFile = new Dictionary<string, WorkspaceFileInfo>(StringComparer.Ordinal);
+			this.NameToSubDirectory = new Dictionary<string, WorkspaceDirectoryInfo>(FileUtils.PlatformPathComparer);
+		}
+
+		public void Read(BinaryReader Reader)
+		{
+			WorkspaceFileInfo[] Files = new WorkspaceFileInfo[Reader.ReadInt32()];
+			for(int Idx = 0; Idx < Files.Length; Idx++)
+			{
+				Files[Idx] = new WorkspaceFileInfo(this, Reader);
+			}
+			NameToFile = Files.ToDictionary(x => x.Name, x => x, StringComparer.Ordinal);
+
+			int NumSubDirectories = Reader.ReadInt32();
+			for(int Idx = 0; Idx < NumSubDirectories; Idx++)
+			{
+				WorkspaceDirectoryInfo SubDirectory = new WorkspaceDirectoryInfo(this, Reader.ReadString());
+				SubDirectory.Read(Reader);
+				NameToSubDirectory[SubDirectory.Name] = SubDirectory;
+			}
+		}
+
+		public void Write(BinaryWriter Writer)
+		{
+			Writer.Write(NameToFile.Count);
+			foreach(WorkspaceFileInfo File in NameToFile.Values)
+			{
+				File.Write(Writer);
+			}
+
+			Writer.Write(NameToSubDirectory.Count);
+			foreach(WorkspaceDirectoryInfo SubDirectory in NameToSubDirectory.Values)
+			{
+				Writer.Write(SubDirectory.Name);
+				SubDirectory.Write(Writer);
+			}
+		}
+
+		public List<WorkspaceFileInfo> GetFiles()
+		{
+			List<WorkspaceFileInfo> Files = new List<WorkspaceFileInfo>();
+			GetFilesInternal(Files);
+			return Files;
+		}
+		
+		private void GetFilesInternal(List<WorkspaceFileInfo> Files)
+		{
+			Files.AddRange(NameToFile.Values);
+
+			foreach(KeyValuePair<string, WorkspaceDirectoryInfo> Pair in NameToSubDirectory)
+			{
+				Pair.Value.GetFilesInternal(Files);
+			}
+		}
+
+		public void Refresh(out FileInfo[] FilesToDelete, out DirectoryInfo[] DirectoriesToDelete)
+		{
+			ConcurrentBag<FileInfo> ConcurrentFilesToDelete = new ConcurrentBag<FileInfo>();
+			ConcurrentBag<DirectoryInfo> ConcurrentDirectoriesToDelete = new ConcurrentBag<DirectoryInfo>();
+			using(ThreadPoolWorkQueue Queue = new ThreadPoolWorkQueue())
+			{
+				Queue.Enqueue(() => Refresh(new DirectoryInfo(GetFullName()), ConcurrentFilesToDelete, ConcurrentDirectoriesToDelete, Queue));
+			}
+			DirectoriesToDelete = ConcurrentDirectoriesToDelete.ToArray();
+			FilesToDelete = ConcurrentFilesToDelete.ToArray();
+		}
+
+		private void Refresh(DirectoryInfo Info, ConcurrentBag<FileInfo> FilesToDelete, ConcurrentBag<DirectoryInfo> DirectoriesToDelete, ThreadPoolWorkQueue Queue)
+		{
+			// Recurse through subdirectories
+			Dictionary<string, WorkspaceDirectoryInfo> NewNameToSubDirectory = new Dictionary<string, WorkspaceDirectoryInfo>(NameToSubDirectory.Count, NameToSubDirectory.Comparer);
+			foreach(DirectoryInfo SubDirectoryInfo in Info.EnumerateDirectories())
+			{
+				WorkspaceDirectoryInfo SubDirectory;
+				if(NameToSubDirectory.TryGetValue(SubDirectoryInfo.Name, out SubDirectory))
+				{
+					NewNameToSubDirectory.Add(SubDirectory.Name, SubDirectory);
+					Queue.Enqueue(() => SubDirectory.Refresh(SubDirectoryInfo, FilesToDelete, DirectoriesToDelete, Queue));
+				}
+				else
+				{
+					DirectoriesToDelete.Add(SubDirectoryInfo);
+				}
+			}
+			NameToSubDirectory = NewNameToSubDirectory;
+			
+			// Figure out which files have changed.
+			Dictionary<string, WorkspaceFileInfo> NewNameToFile = new Dictionary<string, WorkspaceFileInfo>(NameToFile.Count, NameToFile.Comparer);
+			foreach(FileInfo File in Info.EnumerateFiles())
+			{
+				WorkspaceFileInfo StagedFile;
+				if(NameToFile.TryGetValue(File.Name, out StagedFile) && StagedFile.MatchesAttributes(File))
+				{
+					NewNameToFile.Add(StagedFile.Name, StagedFile);
+				}
+				else
+				{
+					FilesToDelete.Add(File);
+				}
+			}
+			NameToFile = NewNameToFile;
+		}
+
+		public string[] FindDifferences()
+		{
+			ConcurrentBag<string> Paths = new ConcurrentBag<string>();
+			using(ThreadPoolWorkQueue Queue = new ThreadPoolWorkQueue())
+			{
+				Queue.Enqueue(() => FindDifferences(new DirectoryInfo(GetFullName()), "/", Paths, Queue));
+			}
+			return Paths.OrderBy(x => x).ToArray();
+		}
+
+		private void FindDifferences(DirectoryInfo Directory, string Path, ConcurrentBag<string> Paths, ThreadPoolWorkQueue Queue)
+		{
+			// Recurse through subdirectories
+			HashSet<string> RemainingSubDirectoryNames = new HashSet<string>(NameToSubDirectory.Keys);
+			foreach(DirectoryInfo SubDirectory in Directory.EnumerateDirectories())
+			{
+				WorkspaceDirectoryInfo StagedSubDirectory;
+				if(NameToSubDirectory.TryGetValue(SubDirectory.Name, out StagedSubDirectory))
+				{
+					RemainingSubDirectoryNames.Remove(SubDirectory.Name);
+					Queue.Enqueue(() => StagedSubDirectory.FindDifferences(SubDirectory, String.Format("{0}{1}/", Path, SubDirectory.Name), Paths, Queue));
+					continue;
+				}
+				Paths.Add(String.Format("+{0}{1}/...", Path, SubDirectory.Name));
+			}
+			foreach(string RemainingSubDirectoryName in RemainingSubDirectoryNames)
+			{
+				Paths.Add(String.Format("-{0}{1}/...", Path, RemainingSubDirectoryName));
+			}
+
+			// Search through files
+			HashSet<string> RemainingFileNames = new HashSet<string>(NameToFile.Keys);
+			foreach(FileInfo File in Directory.EnumerateFiles())
+			{
+				WorkspaceFileInfo StagedFile;
+				if(!NameToFile.TryGetValue(File.Name, out StagedFile))
+				{
+					Paths.Add(String.Format("+{0}{1}", Path, File.Name));
+				}
+				else if(!StagedFile.MatchesAttributes(File))
+				{
+					Paths.Add(String.Format("!{0}{1}", Path, File.Name));
+					RemainingFileNames.Remove(File.Name);
+				}
+				else
+				{
+					RemainingFileNames.Remove(File.Name);
+				}
+			}
+			foreach(string RemainingFileName in RemainingFileNames)
+			{
+				Paths.Add(String.Format("-{0}{1}", Path, RemainingFileName));
+			}
+		}
+
+		public string GetFullName()
+		{
+			StringBuilder Builder = new StringBuilder();
+			AppendFullPath(Builder);
+			return Builder.ToString();
+		}
+
+		public DirectoryReference GetLocation()
+		{
+			return new DirectoryReference(GetFullName());
+		}
+
+		public void AppendClientPath(StringBuilder Builder)
+		{
+			if(ParentDirectory != null)
+			{
+				ParentDirectory.AppendClientPath(Builder);
+				Builder.Append(Name);
+				Builder.Append(Path.DirectorySeparatorChar);
+			}
+		}
+
+		public void AppendFullPath(StringBuilder Builder)
+		{
+			if(ParentDirectory != null)
+			{
+				ParentDirectory.AppendFullPath(Builder);
+				Builder.Append(Path.DirectorySeparatorChar);
+			}
+			Builder.Append(Name);
+		}
+
+		public override string ToString()
+		{
+			return GetFullName();
+		}
+	}
+}
diff --git a/Engine/Source/Programs/BuildAgent/Workspace/Common/WorkspaceTransactionAdd.cs b/Engine/Source/Programs/BuildAgent/Workspace/Common/WorkspaceTransactionAdd.cs
new file mode 100644
--- /dev/null
+++ b/Engine/Source/Programs/BuildAgent/Workspace/Common/WorkspaceTransactionAdd.cs
@@ -0,0 +1,179 @@
+// Copyright 1998-2019 Epic Games, Inc. All Rights Reserved.
+
+using System;
+using System.Collections.Concurrent;
+using System.Collections.Generic;
+using System.IO;
+using System.Linq;
+using System.Text;
+using System.Threading.Tasks;
+using Tools.DotNETCommon;
+
+namespace BuildAgent.Workspace.Common
+{
+	class WorkspaceFileToMove
+	{
+		public StreamFileInfo StreamFile;
+		public TrackedFileInfo TrackedFile;
+		public WorkspaceFileInfo WorkspaceFile;
+
+		public WorkspaceFileToMove(StreamFileInfo StreamFile, TrackedFileInfo TrackedFile, WorkspaceFileInfo WorkspaceFile)
+		{
+			this.StreamFile = StreamFile;
+			this.TrackedFile = TrackedFile;
+			this.WorkspaceFile = WorkspaceFile;
+		}
+	}
+
+	class WorkspaceFileToCopy
+	{
+		public StreamFileInfo StreamFile;
+		public WorkspaceFileInfo SourceWorkspaceFile;
+		public WorkspaceFileInfo TargetWorkspaceFile;
+
+		public WorkspaceFileToCopy(StreamFileInfo StreamFile, WorkspaceFileInfo SourceWorkspaceFile, WorkspaceFileInfo TargetWorkspaceFile)
+		{
+			this.StreamFile = StreamFile;
+			this.SourceWorkspaceFile = SourceWorkspaceFile;
+			this.TargetWorkspaceFile = TargetWorkspaceFile;
+		}
+	}
+
+	class WorkspaceFileToSync
+	{
+		public StreamFileInfo StreamFile;
+		public WorkspaceFileInfo WorkspaceFile;
+
+		public WorkspaceFileToSync(StreamFileInfo StreamFile, WorkspaceFileInfo WorkspaceFile)
+		{
+			this.StreamFile = StreamFile;
+			this.WorkspaceFile = WorkspaceFile;
+		}
+	}
+
+	class WorkspaceTransactionAdd
+	{
+		public WorkspaceDirectoryInfo NewWorkspaceRootDir;
+		public ConcurrentDictionary<TrackedFileInfo, WorkspaceFileToMove> FilesToMove = new ConcurrentDictionary<TrackedFileInfo, WorkspaceFileToMove>();
+		public ConcurrentBag<WorkspaceFileToCopy> FilesToCopy = new ConcurrentBag<WorkspaceFileToCopy>();
+		public ConcurrentBag<WorkspaceFileToSync> FilesToSync = new ConcurrentBag<WorkspaceFileToSync>();
+
+		Dictionary<FileContentId, TrackedFileInfo> ContentIdToTrackedFile;
+		Dictionary<FileContentId, WorkspaceFileInfo> ContentIdToWorkspaceFile = new Dictionary<FileContentId, WorkspaceFileInfo>();
+
+		public WorkspaceTransactionAdd(WorkspaceDirectoryInfo WorkspaceRootDir, StreamDirectoryInfo StreamRootDir, Dictionary<FileContentId, TrackedFileInfo> ContentIdToTrackedFile)
+		{
+			this.NewWorkspaceRootDir = new WorkspaceDirectoryInfo(WorkspaceRootDir.GetLocation());
+
+			this.ContentIdToTrackedFile = ContentIdToTrackedFile;
+
+			List<WorkspaceFileInfo> WorkspaceFiles = WorkspaceRootDir.GetFiles();
+			foreach(WorkspaceFileInfo WorkspaceFile in WorkspaceFiles)
+			{
+				ContentIdToWorkspaceFile[WorkspaceFile.ContentId] = WorkspaceFile;
+			}
+
+			using(ThreadPoolWorkQueue Queue = new ThreadPoolWorkQueue())
+			{
+				Queue.Enqueue(() => MergeDirectory(WorkspaceRootDir, NewWorkspaceRootDir, StreamRootDir, Queue));
+			}
+		}
+
+		void MergeDirectory(WorkspaceDirectoryInfo WorkspaceDir, WorkspaceDirectoryInfo NewWorkspaceDir, StreamDirectoryInfo StreamDir, ThreadPoolWorkQueue Queue)
+		{
+			// Make sure the directory exists
+			Directory.CreateDirectory(WorkspaceDir.GetFullName());
+
+			// Update all the subdirectories
+			foreach(StreamDirectoryInfo StreamSubDir in StreamDir.NameToSubDirectory.Values)
+			{
+				WorkspaceDirectoryInfo WorkspaceSubDir;
+				if(WorkspaceDir.NameToSubDirectory.TryGetValue(StreamSubDir.Name, out WorkspaceSubDir))
+				{
+					MergeSubDirectory(WorkspaceSubDir, StreamSubDir, NewWorkspaceDir, Queue);
+				}
+				else
+				{
+					AddSubDirectory(NewWorkspaceDir, StreamSubDir, Queue);
+				}
+			}
+
+			// Move files into this folder
+			foreach(StreamFileInfo StreamFile in StreamDir.NameToFile.Values)
+			{
+				WorkspaceFileInfo WorkspaceFile;
+				if(WorkspaceDir.NameToFile.TryGetValue(StreamFile.Name, out WorkspaceFile))
+				{
+					NewWorkspaceDir.NameToFile.Add(WorkspaceFile.Name, WorkspaceFile);
+				}
+				else
+				{
+					AddFile(NewWorkspaceDir, StreamFile);
+				}
+			}
+		}
+
+		void AddDirectory(WorkspaceDirectoryInfo NewWorkspaceDir, StreamDirectoryInfo StreamDir, ThreadPoolWorkQueue Queue)
+		{
+			// Make sure the directory exists
+			Directory.CreateDirectory(NewWorkspaceDir.GetFullName());
+
+			// Add all the sub directories and files
+			foreach(StreamDirectoryInfo StreamSubDir in StreamDir.NameToSubDirectory.Values)
+			{
+				AddSubDirectory(NewWorkspaceDir, StreamSubDir, Queue);
+			}
+			foreach(StreamFileInfo StreamFile in StreamDir.NameToFile.Values)
+			{
+				AddFile(NewWorkspaceDir, StreamFile);
+			}
+		}
+
+		void MergeSubDirectory(WorkspaceDirectoryInfo WorkspaceSubDir, StreamDirectoryInfo StreamSubDir, WorkspaceDirectoryInfo NewWorkspaceDir, ThreadPoolWorkQueue Queue)
+		{
+			WorkspaceDirectoryInfo NewWorkspaceSubDir = new WorkspaceDirectoryInfo(NewWorkspaceDir, StreamSubDir.Name);
+			NewWorkspaceDir.NameToSubDirectory.Add(StreamSubDir.Name, NewWorkspaceSubDir);
+			Queue.Enqueue(() => MergeDirectory(WorkspaceSubDir, NewWorkspaceSubDir, StreamSubDir, Queue));
+		}
+
+		void AddSubDirectory(WorkspaceDirectoryInfo NewWorkspaceDir, StreamDirectoryInfo StreamSubDir, ThreadPoolWorkQueue Queue)
+		{
+			WorkspaceDirectoryInfo NewWorkspaceSubDir = new WorkspaceDirectoryInfo(NewWorkspaceDir, StreamSubDir.Name);
+			NewWorkspaceDir.NameToSubDirectory.Add(StreamSubDir.Name, NewWorkspaceSubDir);
+			Queue.Enqueue(() => AddDirectory(NewWorkspaceSubDir, StreamSubDir, Queue));
+		}
+
+		void AddFile(WorkspaceDirectoryInfo NewWorkspaceDir, StreamFileInfo StreamFile)
+		{
+			// Create a new file for this workspace
+			WorkspaceFileInfo WorkspaceFile = new WorkspaceFileInfo(NewWorkspaceDir, StreamFile.Name, StreamFile.ContentId);
+			NewWorkspaceDir.NameToFile.Add(StreamFile.Name, WorkspaceFile);
+
+			// Try to add it to the cache
+			TrackedFileInfo TrackedFile;
+			if(ContentIdToTrackedFile.TryGetValue(StreamFile.ContentId, out TrackedFile))
+			{
+				if(FilesToMove.TryAdd(TrackedFile, new WorkspaceFileToMove(StreamFile, TrackedFile, WorkspaceFile)))
+				{
+					WorkspaceFile.SetMetadata(TrackedFile.Length, TrackedFile.LastModifiedTicks, TrackedFile.bReadOnly);
+				}
+				else
+				{
+					FilesToCopy.Add(new WorkspaceFileToCopy(StreamFile, FilesToMove[TrackedFile].WorkspaceFile, WorkspaceFile));
+				}
+			}
+			else
+			{
+				WorkspaceFileInfo SourceWorkspaceFile;
+				if(ContentIdToWorkspaceFile.TryGetValue(StreamFile.ContentId, out SourceWorkspaceFile))
+				{
+					FilesToCopy.Add(new WorkspaceFileToCopy(StreamFile, SourceWorkspaceFile, WorkspaceFile));
+				}
+				else
+				{
+					FilesToSync.Add(new WorkspaceFileToSync(StreamFile, WorkspaceFile));
+				}
+			}
+		}
+	}
+}
diff --git a/Engine/Source/Programs/BuildAgent/Workspace/Common/WorkspaceTransactionRemove.cs b/Engine/Source/Programs/BuildAgent/Workspace/Common/WorkspaceTransactionRemove.cs
new file mode 100644
--- /dev/null
+++ b/Engine/Source/Programs/BuildAgent/Workspace/Common/WorkspaceTransactionRemove.cs
@@ -0,0 +1,93 @@
+// Copyright 1998-2019 Epic Games, Inc. All Rights Reserved.
+
+using System;
+using System.Collections.Concurrent;
+using System.Collections.Generic;
+using System.Linq;
+using System.Text;
+using System.Threading.Tasks;
+using Tools.DotNETCommon;
+
+namespace BuildAgent.Workspace.Common
+{
+	class WorkspaceTransactionRemove
+	{
+		Dictionary<FileContentId, TrackedFileInfo> ContentIdToTrackedFile;
+
+		public WorkspaceDirectoryInfo NewWorkspaceRootDir;
+		public ConcurrentDictionary<FileContentId, WorkspaceFileInfo> FilesToMove = new ConcurrentDictionary<FileContentId, WorkspaceFileInfo>();
+		public ConcurrentBag<WorkspaceFileInfo> FilesToDelete = new ConcurrentBag<WorkspaceFileInfo>();
+		public ConcurrentBag<WorkspaceDirectoryInfo> DirectoriesToDelete = new ConcurrentBag<WorkspaceDirectoryInfo>();
+
+		public WorkspaceTransactionRemove(WorkspaceDirectoryInfo WorkspaceRootDir, StreamDirectoryInfo StreamRootDir, Dictionary<FileContentId, TrackedFileInfo> ContentIdToTrackedFile)
+		{
+			this.NewWorkspaceRootDir = new WorkspaceDirectoryInfo(WorkspaceRootDir.GetLocation());
+			this.ContentIdToTrackedFile = ContentIdToTrackedFile;
+
+			using(ThreadPoolWorkQueue Queue = new ThreadPoolWorkQueue())
+			{
+				Queue.Enqueue(() => Merge(WorkspaceRootDir, NewWorkspaceRootDir, StreamRootDir, Queue));
+			}
+		}
+
+		void Merge(WorkspaceDirectoryInfo WorkspaceDir, WorkspaceDirectoryInfo NewWorkspaceDir, StreamDirectoryInfo StreamDir, ThreadPoolWorkQueue Queue)
+		{
+			// Update all the subdirectories
+			foreach(WorkspaceDirectoryInfo WorkspaceSubDir in WorkspaceDir.NameToSubDirectory.Values)
+			{
+				StreamDirectoryInfo StreamSubDir;
+				if(StreamDir.NameToSubDirectory.TryGetValue(WorkspaceSubDir.Name, out StreamSubDir))
+				{
+					MergeSubDirectory(NewWorkspaceDir, WorkspaceSubDir, StreamSubDir, Queue);
+				}
+				else
+				{
+					RemoveDirectory(WorkspaceSubDir, Queue);
+				}
+			}
+
+			// Update the staged files
+			foreach(WorkspaceFileInfo WorkspaceFile in WorkspaceDir.NameToFile.Values)
+			{
+				StreamFileInfo StreamFile;
+				if(StreamDir != null && StreamDir.NameToFile.TryGetValue(WorkspaceFile.Name, out StreamFile) && StreamFile.ContentId.Equals(WorkspaceFile.ContentId))
+				{
+					NewWorkspaceDir.NameToFile.Add(WorkspaceFile.Name, WorkspaceFile);
+				}
+				else
+				{
+					RemoveFile(WorkspaceFile);
+				}
+			}
+		}
+
+		void MergeSubDirectory(WorkspaceDirectoryInfo NewWorkspaceDir, WorkspaceDirectoryInfo WorkspaceSubDir, StreamDirectoryInfo StreamSubDir, ThreadPoolWorkQueue Queue)
+		{
+			WorkspaceDirectoryInfo NewWorkspaceSubDir = new WorkspaceDirectoryInfo(NewWorkspaceDir, StreamSubDir.Name);
+			NewWorkspaceDir.NameToSubDirectory.Add(NewWorkspaceSubDir.Name, NewWorkspaceSubDir);
+			Queue.Enqueue(() => Merge(WorkspaceSubDir, NewWorkspaceSubDir, StreamSubDir, Queue));
+		}
+
+		void RemoveDirectory(WorkspaceDirectoryInfo WorkspaceDir, ThreadPoolWorkQueue Queue)
+		{
+			DirectoriesToDelete.Add(WorkspaceDir);
+
+			foreach(WorkspaceDirectoryInfo WorkspaceSubDir in WorkspaceDir.NameToSubDirectory.Values)
+			{
+				Queue.Enqueue(() => RemoveDirectory(WorkspaceSubDir, Queue));
+			}
+			foreach(WorkspaceFileInfo WorkspaceFile in WorkspaceDir.NameToFile.Values)
+			{
+				RemoveFile(WorkspaceFile);
+			}
+		}
+
+		void RemoveFile(WorkspaceFileInfo WorkspaceFile)
+		{
+			if(ContentIdToTrackedFile.ContainsKey(WorkspaceFile.ContentId) || !FilesToMove.TryAdd(WorkspaceFile.ContentId, WorkspaceFile))
+			{
+				FilesToDelete.Add(WorkspaceFile);
+			}
+		}
+	}
+}
diff --git a/Engine/Source/Programs/BuildAgent/Workspace/DumpMode.cs b/Engine/Source/Programs/BuildAgent/Workspace/DumpMode.cs
new file mode 100644
--- /dev/null
+++ b/Engine/Source/Programs/BuildAgent/Workspace/DumpMode.cs
@@ -0,0 +1,20 @@
+// Copyright 1998-2019 Epic Games, Inc. All Rights Reserved.
+
+using BuildAgent.Workspace.Common;
+using System;
+using System.Collections.Generic;
+using System.Linq;
+using System.Text;
+using System.Threading.Tasks;
+
+namespace BuildAgent.Workspace
+{
+	[ProgramMode("Dump", "Dumps the contents of the repository to the log for analysis")]
+	class DumpMode : WorkspaceMode
+	{
+		protected override void Execute(Repository Repo)
+		{
+			Repo.Dump();
+		}
+	}
+}
diff --git a/Engine/Source/Programs/BuildAgent/Workspace/PopulateCacheMode.cs b/Engine/Source/Programs/BuildAgent/Workspace/PopulateCacheMode.cs
new file mode 100644
--- /dev/null
+++ b/Engine/Source/Programs/BuildAgent/Workspace/PopulateCacheMode.cs
@@ -0,0 +1,37 @@
+// Copyright 1998-2019 Epic Games, Inc. All Rights Reserved.
+
+using BuildAgent.Workspace.Common;
+using System;
+using System.Collections.Generic;
+using System.ComponentModel;
+using System.Linq;
+using System.Text;
+using System.Threading.Tasks;
+using Tools.DotNETCommon;
+
+namespace BuildAgent.Workspace
+{
+	[ProgramMode("PopulateCache", "Populates the cache with the head revision of the given streams")]
+	class PopulateCacheMode : WorkspaceMode
+	{
+		[CommandLine("-ClientAndStream=")]
+		[Description("Specifies client and stream pairs, in the format Client:Stream")]
+		List<string> ClientAndStreamParams = new List<string>();
+
+		[CommandLine("-Filter=")]
+		[Description("Filters for the files to sync, in P4 syntax (eg. /Engine/...)")]
+		List<string> Filters = new List<string>();
+
+		[CommandLine("-FakeSync")]
+		[Description("Simulates the sync without actually fetching any files")]
+		bool bFakeSync = false;
+
+		protected override void Execute(Repository Repo)
+		{
+			List<string> ExpandedFilters = ExpandFilters(Filters);
+
+			List<KeyValuePair<string, string>> ClientAndStreams = ParseClientAndStreams(ClientAndStreamParams);
+			Repo.Populate(ClientAndStreams, Filters, bFakeSync);
+		}
+	}
+}
diff --git a/Engine/Source/Programs/BuildAgent/Workspace/PurgeCacheMode.cs b/Engine/Source/Programs/BuildAgent/Workspace/PurgeCacheMode.cs
new file mode 100644
--- /dev/null
+++ b/Engine/Source/Programs/BuildAgent/Workspace/PurgeCacheMode.cs
@@ -0,0 +1,30 @@
+// Copyright 1998-2019 Epic Games, Inc. All Rights Reserved.
+
+using BuildAgent.Workspace.Common;
+using System;
+using System.Collections.Generic;
+using System.Linq;
+using System.Text;
+using System.Threading.Tasks;
+using Tools.DotNETCommon;
+
+namespace BuildAgent.Workspace
+{
+	[ProgramMode("PurgeCache", "Shrink the size of the cache to the given size")]
+	class PurgeCacheMode : WorkspaceMode
+	{
+		[CommandLine("-Size=")]
+		string SizeParam = null;
+
+		protected override void Execute(Repository Repo)
+		{
+			long Size = 0;
+			if (SizeParam != null)
+			{
+				Size = ParseSize(SizeParam);
+			}
+
+			Repo.Purge(Size);
+		}
+	}
+}
diff --git a/Engine/Source/Programs/BuildAgent/Workspace/RepairCacheMode.cs b/Engine/Source/Programs/BuildAgent/Workspace/RepairCacheMode.cs
new file mode 100644
--- /dev/null
+++ b/Engine/Source/Programs/BuildAgent/Workspace/RepairCacheMode.cs
@@ -0,0 +1,20 @@
+// Copyright 1998-2019 Epic Games, Inc. All Rights Reserved.
+
+using BuildAgent.Workspace.Common;
+using System;
+using System.Collections.Generic;
+using System.Linq;
+using System.Text;
+using System.Threading.Tasks;
+
+namespace BuildAgent.Workspace
+{
+	[ProgramMode("RepairCache", "Checks the integrity of the cache, and removes any invalid files")]
+	class RepairCacheMode : WorkspaceMode
+	{
+		protected override void Execute(Repository Repo)
+		{
+			Repo.Repair();
+		}
+	}
+}
diff --git a/Engine/Source/Programs/BuildAgent/Workspace/RevertMode.cs b/Engine/Source/Programs/BuildAgent/Workspace/RevertMode.cs
new file mode 100644
--- /dev/null
+++ b/Engine/Source/Programs/BuildAgent/Workspace/RevertMode.cs
@@ -0,0 +1,26 @@
+// Copyright 1998-2019 Epic Games, Inc. All Rights Reserved.
+
+using BuildAgent.Workspace.Common;
+using System;
+using System.Collections.Generic;
+using System.ComponentModel;
+using System.Linq;
+using System.Text;
+using System.Threading.Tasks;
+using Tools.DotNETCommon;
+
+namespace BuildAgent.Workspace
+{
+	[ProgramMode("Revert", "Revert all files that are open in the current workspace. Does not replace them with valid revisions.")]
+	class RevertMode : WorkspaceMode
+	{
+		[CommandLine("-Client=", Required = true)]
+		[Description("Client to revert all files for")]
+		string ClientName = null;
+
+		protected override void Execute(Repository Repo)
+		{
+			Repo.Revert(ClientName);
+		}
+	}
+}
diff --git a/Engine/Source/Programs/BuildAgent/Workspace/SetupMode.cs b/Engine/Source/Programs/BuildAgent/Workspace/SetupMode.cs
new file mode 100644
--- /dev/null
+++ b/Engine/Source/Programs/BuildAgent/Workspace/SetupMode.cs
@@ -0,0 +1,30 @@
+// Copyright 1998-2019 Epic Games, Inc. All Rights Reserved.
+
+using System;
+using System.Collections.Generic;
+using System.ComponentModel;
+using System.Linq;
+using System.Text;
+using System.Threading.Tasks;
+using BuildAgent.Workspace.Common;
+using Tools.DotNETCommon;
+
+namespace BuildAgent.Workspace
+{
+	[ProgramMode("Setup", "Creates or updates a client to use a given stream")]
+	class SetupMode : WorkspaceMode
+	{
+		[CommandLine("-Client=", Required = true)]
+		[Description("Name of the client to create")]
+		string ClientName = null;
+
+		[CommandLine("-Stream=", Required = true)]
+		[Description("Name of the stream to configure")]
+		string StreamName = null;
+
+		protected override void Execute(Repository Repo)
+		{
+			Repo.Setup(ClientName, StreamName);
+		}
+	}
+}
diff --git a/Engine/Source/Programs/BuildAgent/Workspace/SizesMode.cs b/Engine/Source/Programs/BuildAgent/Workspace/SizesMode.cs
new file mode 100644
--- /dev/null
+++ b/Engine/Source/Programs/BuildAgent/Workspace/SizesMode.cs
@@ -0,0 +1,34 @@
+// Copyright 1998-2019 Epic Games, Inc. All Rights Reserved.
+
+using System;
+using System.Collections.Generic;
+using System.ComponentModel;
+using System.Linq;
+using System.Text;
+using System.Threading.Tasks;
+using BuildAgent.Workspace.Common;
+using Tools.DotNETCommon;
+
+namespace BuildAgent.Workspace
+{
+	[ProgramMode("Sizes", "Gathers stats for which streams take the most amount of space in the cache for the given configuration")]
+	class SizesMode : WorkspaceMode
+	{
+		[CommandLine("-TempClient=", Required = true)]
+		[Description("Name of a temporary client to switch between streams gathering metadata. Will be created if it does not exist.")]
+		string TempClientName = null;
+
+		[CommandLine("-Stream=")]
+		[Description("Streams that should be included in the output")]
+		List<string> StreamNames = new List<string>();
+
+		[CommandLine("-Filter=")]
+		[Description("Filters for the files to sync, in P4 syntax (eg. /Engine/...)")]
+		List<string> Filters = new List<string>();
+
+		protected override void Execute(Repository Repo)
+		{
+			Repo.Stats(TempClientName, StreamNames, Filters);
+		}
+	}
+}
diff --git a/Engine/Source/Programs/BuildAgent/Workspace/StatusMode.cs b/Engine/Source/Programs/BuildAgent/Workspace/StatusMode.cs
new file mode 100644
--- /dev/null
+++ b/Engine/Source/Programs/BuildAgent/Workspace/StatusMode.cs
@@ -0,0 +1,20 @@
+// Copyright 1998-2019 Epic Games, Inc. All Rights Reserved.
+
+using BuildAgent.Workspace.Common;
+using System;
+using System.Collections.Generic;
+using System.Linq;
+using System.Text;
+using System.Threading.Tasks;
+
+namespace BuildAgent.Workspace
+{
+	[ProgramMode("Status", "Prints information about the state of the cache and workspace")]
+	class StatusMode : WorkspaceMode
+	{
+		protected override void Execute(Repository Repo)
+		{
+			Repo.Status();
+		}
+	}
+}
diff --git a/Engine/Source/Programs/BuildAgent/Workspace/SyncMode.cs b/Engine/Source/Programs/BuildAgent/Workspace/SyncMode.cs
new file mode 100644
--- /dev/null
+++ b/Engine/Source/Programs/BuildAgent/Workspace/SyncMode.cs
@@ -0,0 +1,49 @@
+// Copyright 1998-2019 Epic Games, Inc. All Rights Reserved.
+
+using BuildAgent.Workspace.Common;
+using System;
+using System.Collections.Generic;
+using System.ComponentModel;
+using System.Linq;
+using System.Text;
+using System.Threading.Tasks;
+using Tools.DotNETCommon;
+
+namespace BuildAgent.Workspace
+{
+	[ProgramMode("Sync", "Syncs the files for a particular stream and changelist")]
+	class SyncMode : WorkspaceMode
+	{
+		[CommandLine("-Client=", Required = true)]
+		[Description("Name of the client to sync. Will be created if it does not exist.")]
+		public string ClientName = null;
+
+		[CommandLine("-Stream=", Required = true)]
+		[Description("The stream to sync")]
+		public string StreamName = null;
+
+		[CommandLine("-Change=")]
+		[Description("The change to sync. May be a changelist number, or 'Latest'")]
+		public string Change = "Latest";
+
+		[CommandLine]
+		[Description("Optional path to a cache file used to store workspace metadata. Using a location on a network share allows multiple machines syncing the same CL to only query Perforce state once.")]
+		FileReference CacheFile = null;
+
+		[CommandLine("-Filter=")]
+		[Description("Filters for the files to sync, in P4 syntax (eg. /Engine/...)")]
+		public List<string> Filters = new List<string>();
+
+		[CommandLine("-FakeSync")]
+		[Description("Simulates the sync without actually fetching any files")]
+		public bool bFakeSync = false;
+
+		protected override void Execute(Repository Repo)
+		{
+			int ChangeNumber = ParseChangeNumberOrLatest(Change);
+			List<string> ExpandedFilters = ExpandFilters(Filters);
+
+			Repo.Sync(ClientName, StreamName, ChangeNumber, ExpandedFilters, bFakeSync, CacheFile);
+		}
+	}
+}
diff --git a/Engine/Source/Programs/BuildAgent/Workspace/WorkspaceMode.cs b/Engine/Source/Programs/BuildAgent/Workspace/WorkspaceMode.cs
new file mode 100644
--- /dev/null
+++ b/Engine/Source/Programs/BuildAgent/Workspace/WorkspaceMode.cs
@@ -0,0 +1,236 @@
+// Copyright 1998-2019 Epic Games, Inc. All Rights Reserved.
+
+using BuildAgent.Workspace.Common;
+using System;
+using System.Collections.Generic;
+using System.ComponentModel;
+using System.Diagnostics;
+using System.IO;
+using System.Linq;
+using System.Text;
+using System.Threading.Tasks;
+using Tools.DotNETCommon;
+
+namespace BuildAgent.Workspace
+{
+	abstract class WorkspaceMode : ProgramMode
+	{
+		[CommandLine("-Server")]
+		[Description("Specifies the Perforce server and port")]
+		protected string ServerAndPort = null;
+
+		[CommandLine("-User")]
+		[Description("Specifies the Perforce username")]
+		protected string UserName = null;
+
+		[CommandLine("-BaseDir")]
+		[Description("Base directory to use for syncing workspaces")]
+		protected DirectoryReference BaseDir = null;
+
+		[CommandLine("-Overwrite")]
+		[Description("")]
+		protected bool bOverwrite = false;
+
+		[CommandLine("-Verbose")]
+		[Description("Enables verbose logging")]
+		protected bool bVerbose = false;
+
+		public override void Configure(CommandLineArguments Arguments)
+		{
+			base.Configure(Arguments);
+
+			if(BaseDir == null)
+			{
+				for (DirectoryReference ParentDir = DirectoryReference.GetCurrentDirectory(); ParentDir != null; ParentDir = ParentDir.ParentDirectory)
+				{
+					if (Repository.Exists(ParentDir))
+					{
+						BaseDir = ParentDir;
+						break;
+					}
+				}
+
+				if (BaseDir == null)
+				{
+					throw new FatalErrorException("Unable to find existing repository in current working tree. Specify -BaseDir=... explicitly.");
+				}
+			}
+		}
+
+		public override void Execute()
+		{
+			// Create the repo
+			Repository Repo = CreateOrLoadRepository(ServerAndPort, UserName, BaseDir, bOverwrite);
+
+			// Delete any old log files
+			FileReference LogFile = FileReference.Combine(BaseDir, "Logs", "Log.txt");
+			DirectoryReference.CreateDirectory(LogFile.Directory);
+			BackupLogFile(LogFile, TimeSpan.FromDays(3));
+			if (bVerbose)
+			{
+				Log.OutputLevel = LogEventType.Verbose;
+			}
+			Trace.Listeners.Add(new TextWriterTraceListener(new StreamWriter(LogFile.FullName), "LogTraceListener"));
+
+			Execute(Repo);
+		}
+
+		protected abstract void Execute(Repository Repo);
+
+		static void BackupLogFile(FileReference LogFile, TimeSpan RetentionSpan)
+		{
+			string LogPrefix = LogFile.GetFileNameWithoutExtension() + "-backup-";
+			string LogSuffix = LogFile.GetExtension();
+
+			if (FileReference.Exists(LogFile))
+			{
+				string Timestamp = FileReference.GetLastWriteTime(LogFile).ToString("yyyy.MM.dd-HH.mm.ss");
+
+				string UniqueSuffix = "";
+				for (int NextSuffixIdx = 2; ; NextSuffixIdx++)
+				{
+					FileReference BackupLogFile = FileReference.Combine(LogFile.Directory, String.Format("{0}{1}{2}{3}", LogPrefix, Timestamp, UniqueSuffix, LogSuffix));
+					if (!FileReference.Exists(BackupLogFile))
+					{
+						try
+						{
+							FileReference.Move(LogFile, BackupLogFile);
+						}
+						catch (Exception Ex)
+						{
+							Log.TraceWarning("Unable to backup {0} to {1} ({2})", LogFile, BackupLogFile, Ex.Message);
+							Log.TraceLog(ExceptionUtils.FormatException(Ex));
+						}
+						break;
+					}
+					UniqueSuffix = String.Format("_{0}", NextSuffixIdx);
+				}
+			}
+
+			DateTime DeleteTime = DateTime.UtcNow - RetentionSpan;
+			foreach (FileInfo File in new DirectoryInfo(LogFile.Directory.FullName).EnumerateFiles(LogPrefix + "*" + LogSuffix))
+			{
+				if (File.LastWriteTimeUtc < DeleteTime)
+				{
+					File.Delete();
+				}
+			}
+		}
+
+		static Repository CreateOrLoadRepository(string ServerAndPort, string UserName, DirectoryReference BaseDir, bool bOverwrite)
+		{
+			if (Repository.Exists(BaseDir))
+			{
+				try
+				{
+					return Repository.Load(ServerAndPort, UserName, BaseDir);
+				}
+				catch (Exception Ex)
+				{
+					if (bOverwrite)
+					{
+						Log.TraceWarning("Unable to load existing repository: {0}", Ex.Message);
+						Log.TraceVerbose("{0}", ExceptionUtils.FormatExceptionDetails(Ex));
+					}
+					else
+					{
+						throw;
+					}
+
+				}
+			}
+
+			return Repository.Create(ServerAndPort, UserName, BaseDir);
+		}
+
+		protected static long ParseSize(string Size)
+		{
+			long Value;
+			if (Size.EndsWith("gb", StringComparison.OrdinalIgnoreCase))
+			{
+				string SizeValue = Size.Substring(0, Size.Length - 2).TrimEnd();
+				if (long.TryParse(SizeValue, out Value))
+				{
+					return Value * (1024 * 1024 * 1024);
+				}
+			}
+			else if (Size.EndsWith("mb", StringComparison.OrdinalIgnoreCase))
+			{
+				string SizeValue = Size.Substring(0, Size.Length - 2).TrimEnd();
+				if (long.TryParse(SizeValue, out Value))
+				{
+					return Value * (1024 * 1024);
+				}
+			}
+			else if (Size.EndsWith("kb"))
+			{
+				string SizeValue = Size.Substring(0, Size.Length - 2).TrimEnd();
+				if (long.TryParse(SizeValue, out Value))
+				{
+					return Value * 1024;
+				}
+			}
+			else
+			{
+				if (long.TryParse(Size, out Value))
+				{
+					return Value;
+				}
+			}
+			throw new FatalErrorException("Invalid size '{0}'", Size);
+		}
+
+		protected static int ParseChangeNumber(string Change)
+		{
+			int ChangeNumber;
+			if (int.TryParse(Change, out ChangeNumber) && ChangeNumber > 0)
+			{
+				return ChangeNumber;
+			}
+			throw new FatalErrorException("Unable to parse change number from '{0}'", Change);
+		}
+
+		protected static int ParseChangeNumberOrLatest(string Change)
+		{
+			if (Change.Equals("Latest", StringComparison.OrdinalIgnoreCase))
+			{
+				return -1;
+			}
+			else
+			{
+				return ParseChangeNumber(Change);
+			}
+		}
+
+		protected static List<KeyValuePair<string, string>> ParseClientAndStreams(List<string> ClientAndStreamParams)
+		{
+			List<KeyValuePair<string, string>> ClientAndStreams = new List<KeyValuePair<string, string>>();
+			foreach (string ClientAndStreamParam in ClientAndStreamParams)
+			{
+				int Idx = ClientAndStreamParam.IndexOf(':');
+				if (Idx == -1)
+				{
+					throw new FatalErrorException("Expected -ClientAndStream=<ClientName>:<StreamName>");
+				}
+				ClientAndStreams.Add(new KeyValuePair<string, string>(ClientAndStreamParam.Substring(0, Idx), ClientAndStreamParam.Substring(Idx + 1)));
+			}
+			return ClientAndStreams;
+		}
+
+		protected static List<string> ExpandFilters(List<string> Arguments)
+		{
+			List<string> Filters = new List<string>();
+			foreach (string Argument in Arguments)
+			{
+				foreach (string SingleArgument in Argument.Split(';').Select(x => x.Trim()))
+				{
+					if (SingleArgument.Length > 0)
+					{
+						Filters.Add(SingleArgument);
+					}
+				}
+			}
+			return Filters;
+		}
+	}
+}
diff --git a/Engine/Source/Programs/DotNETCommon/DotNETUtilities/CommandLineArguments.cs b/Engine/Source/Programs/DotNETCommon/DotNETUtilities/CommandLineArguments.cs
--- a/Engine/Source/Programs/DotNETCommon/DotNETUtilities/CommandLineArguments.cs
+++ b/Engine/Source/Programs/DotNETCommon/DotNETUtilities/CommandLineArguments.cs
@@ -757,6 +757,42 @@ namespace Tools.DotNETCommon
 			return Parameters;
 		}
 
+		/// <summary>
+		/// Quotes a command line argument, if necessary
+		/// </summary>
+		/// <param name="Argument">The argument that may need quoting</param>
+		/// <returns>Argument which is safe to pass on the command line</returns>
+		public static string Quote(string Argument)
+		{
+			if (Argument.IndexOf(' ') != -1 && Argument[0] != '\"')
+			{
+				return "\"" + Argument+ "\"";
+			}
+			else
+			{
+				return Argument;
+			}
+		}
+
+		/// <summary>
+		/// Joins the given arguments into a command line
+		/// </summary>
+		/// <param name="Arguments">List of command line arguments</param>
+		/// <returns>Joined command line</returns>
+		public static string Join(IEnumerable<string> Arguments)
+		{
+			StringBuilder Result = new StringBuilder();
+			foreach (string Argument in Arguments)
+			{
+				if(Result.Length > 0)
+				{
+					Result.Append(' ');
+				}
+				Result.Append(Quote(Argument));
+			}
+			return Result.ToString();
+		}
+
 		/// <summary>
 		/// Splits a command line into individual arguments
 		/// </summary>
@@ -969,7 +1005,7 @@ namespace Tools.DotNETCommon
 		/// <returns>True if the text could be parsed, false otherwise</returns>
 		private static bool TryParseValue(Type FieldType, string Text, out object Value)
 		{
-			if(FieldType.IsGenericType && FieldType.GetGenericTypeDefinition() == typeof(Nullable<>))
+			if (FieldType.IsGenericType && FieldType.GetGenericTypeDefinition() == typeof(Nullable<>))
 			{
 				// Try to parse the inner type instead
 				return TryParseValue(FieldType.GetGenericArguments()[0], Text, out Value);
@@ -982,13 +1018,13 @@ namespace Tools.DotNETCommon
 					Value = Enum.Parse(FieldType, Text, true);
 					return true;
 				}
-				catch(ArgumentException)
+				catch (ArgumentException)
 				{
 					Value = null;
 					return false;
 				}
 			}
-			else if(FieldType == typeof(FileReference))
+			else if (FieldType == typeof(FileReference))
 			{
 				// Construct a file reference from the string
 				try
@@ -1002,7 +1038,7 @@ namespace Tools.DotNETCommon
 					return false;
 				}
 			}
-			else if(FieldType == typeof(DirectoryReference))
+			else if (FieldType == typeof(DirectoryReference))
 			{
 				// Construct a file reference from the string
 				try
@@ -1016,6 +1052,38 @@ namespace Tools.DotNETCommon
 					return false;
 				}
 			}
+			else if (FieldType == typeof(TimeSpan))
+			{
+				// Construct a time span form the string
+				double FloatValue;
+				if (Text.EndsWith("h", StringComparison.OrdinalIgnoreCase) && Double.TryParse(Text.Substring(0, Text.Length - 1), out FloatValue))
+				{
+					Value = TimeSpan.FromHours(FloatValue);
+					return true;
+				}
+				else if (Text.EndsWith("m", StringComparison.OrdinalIgnoreCase) && Double.TryParse(Text.Substring(0, Text.Length - 1), out FloatValue))
+				{
+					Value = TimeSpan.FromMinutes(FloatValue);
+					return true;
+				}
+				else if (Text.EndsWith("s", StringComparison.OrdinalIgnoreCase) && Double.TryParse(Text.Substring(0, Text.Length - 1), out FloatValue))
+				{
+					Value = TimeSpan.FromSeconds(FloatValue);
+					return true;
+				}
+
+				TimeSpan TimeSpanValue;
+				if (TimeSpan.TryParse(Text, out TimeSpanValue))
+				{
+					Value = TimeSpanValue;
+					return true;
+				}
+				else
+				{
+					Value = null;
+					return false;
+				}
+			}
 			else
 			{
 				// Otherwise let the framework convert between types
@@ -1024,7 +1092,7 @@ namespace Tools.DotNETCommon
 					Value = Convert.ChangeType(Text, FieldType);
 					return true;
 				}
-				catch(InvalidCastException)
+				catch (InvalidCastException)
 				{
 					Value = null;
 					return false;
diff --git a/Engine/Source/Programs/DotNETCommon/DotNETUtilities/DotNETUtilities.csproj b/Engine/Source/Programs/DotNETCommon/DotNETUtilities/DotNETUtilities.csproj
--- a/Engine/Source/Programs/DotNETCommon/DotNETUtilities/DotNETUtilities.csproj
+++ b/Engine/Source/Programs/DotNETCommon/DotNETUtilities/DotNETUtilities.csproj
@@ -76,6 +76,7 @@
     <Compile Include="FastJSON\JsonSerializer.cs" />
     <Compile Include="FastJSON\Reflection.cs" />
     <Compile Include="FastJSON\SafeDictionary.cs" />
+    <Compile Include="FatalErrorException.cs" />
     <Compile Include="FileFilter.cs" />
     <Compile Include="FileFilterNode.cs" />
     <Compile Include="FilePattern.cs" />
@@ -156,6 +157,7 @@
     <Compile Include="Perforce\Records\UpdateChangeOptions.cs" />
     <Compile Include="Perforce\Records\UserRecord.cs" />
     <Compile Include="Perforce\Records\WhereRecord.cs" />
+    <Compile Include="ProcessUtils.cs" />
     <Compile Include="Properties\AssemblyInfo.cs" />
     <Compile Include="ReadOnlyHashSet.cs" />
     <Compile Include="ResXResource\UEResXReader.cs" />
diff --git a/Engine/Source/Programs/DotNETCommon/DotNETUtilities/FatalErrorException.cs b/Engine/Source/Programs/DotNETCommon/DotNETUtilities/FatalErrorException.cs
new file mode 100644
--- /dev/null
+++ b/Engine/Source/Programs/DotNETCommon/DotNETUtilities/FatalErrorException.cs
@@ -0,0 +1,70 @@
+// Copyright 1998-2019 Epic Games, Inc. All Rights Reserved.
+
+using System;
+using System.Collections.Generic;
+using System.Linq;
+using System.Text;
+using System.Threading.Tasks;
+
+namespace Tools.DotNETCommon
+{
+	/// <summary>
+	/// Base class for user-facing fatal errors. These errors are shown to the user prior to termination, without a callstack, and may dictate the program exit code.
+	/// </summary>
+	public class FatalErrorException : Exception
+	{
+		/// <summary>
+		/// Exit code for the process
+		/// </summary>
+		public int ExitCode = 1;
+
+		/// <summary>
+		/// Constructor
+		/// </summary>
+		/// <param name="Message">The error message to display.</param>
+		public FatalErrorException(string Message)
+			: base(Message)
+		{
+		}
+
+		/// <summary>
+		/// Constructor
+		/// </summary>
+		/// <param name="InnerException">An inner exception to wrap</param>
+		/// <param name="Message">The error message to display.</param>
+		public FatalErrorException(Exception InnerException, string Message)
+			: base(Message, InnerException)
+		{
+		}
+
+		/// <summary>
+		/// Constructor
+		/// </summary>
+		/// <param name="Format">Formatting string for the error message</param>
+		/// <param name="Arguments">Arguments for the formatting string</param>
+		public FatalErrorException(string Format, params object[] Arguments)
+			: base(String.Format(Format, Arguments))
+		{
+		}
+
+		/// <summary>
+		/// Constructor which wraps another exception
+		/// </summary>
+		/// <param name="InnerException">The inner exception being wrapped</param>
+		/// <param name="Format">Format for the message string</param>
+		/// <param name="Arguments">Format arguments</param>
+		public FatalErrorException(Exception InnerException, string Format, params object[] Arguments)
+			: base(String.Format(Format, Arguments), InnerException)
+		{
+		}
+
+		/// <summary>
+		/// Returns the string representing the exception. Our build exceptions do not show the callstack since they are used to report known error conditions.
+		/// </summary>
+		/// <returns>Message for the exception</returns>
+		public override string ToString()
+		{
+			return Message;
+		}
+	}
+}
diff --git a/Engine/Source/Programs/DotNETCommon/DotNETUtilities/FileUtils.cs b/Engine/Source/Programs/DotNETCommon/DotNETUtilities/FileUtils.cs
--- a/Engine/Source/Programs/DotNETCommon/DotNETUtilities/FileUtils.cs
+++ b/Engine/Source/Programs/DotNETCommon/DotNETUtilities/FileUtils.cs
@@ -2,18 +2,49 @@
 
 using System;
 using System.Collections.Generic;
+using System.ComponentModel;
 using System.IO;
 using System.Linq;
+using System.Runtime.InteropServices;
 using System.Text;
 using System.Threading.Tasks;
 
 namespace Tools.DotNETCommon
 {
 	/// <summary>
-	/// Utility functions for manipulating files
+	/// Exception used to represent caught file/directory exceptions.
+	/// </summary>
+	class WrappedFileOrDirectoryException : Exception
+	{
+		/// <summary>
+		/// Constructor
+		/// </summary>
+		/// <param name="Inner">Inner exception</param>
+		/// <param name="Message">Message to display</param>
+		public WrappedFileOrDirectoryException(Exception Inner, string Message) : base(Message)
+		{
+		}
+
+		/// <summary>
+		/// Returns the message to display for this exception
+		/// </summary>
+		/// <returns>Message to display</returns>
+		public override string ToString()
+		{
+			return Message;
+		}
+	}
+
+	/// <summary>
+	/// Utility functions for manipulating files. Where these methods have similar functionality to those in the NET Framework, they generally go the extra mile to produce concise, specific error messages where possible.
 	/// </summary>
 	public static class FileUtils
 	{
+		/// <summary>
+		/// Comparer that should be used for native path comparisons
+		/// </summary>
+		public static IEqualityComparer<string> PlatformPathComparer = Environment.OSVersion.Platform == PlatformID.Unix ? StringComparer.Ordinal : StringComparer.OrdinalIgnoreCase;
+
 		/// <summary>
 		/// Finds the on-disk case of a a file
 		/// </summary>
@@ -31,5 +62,313 @@ namespace Tools.DotNETCommon
 			}
 			return new FileInfo(Path.Combine(ParentInfo.FullName, Info.Name));
 		}
+
+		/// <summary>
+		/// Creates a directory tree, with all intermediate branches
+		/// </summary>
+		/// <param name="Directory">The directory to create</param>
+		public static void CreateDirectoryTree(DirectoryReference Directory)
+		{
+			if (!DirectoryReference.Exists(Directory))
+			{
+				CreateDirectoryTree(Directory.ParentDirectory);
+				DirectoryReference.CreateDirectory(Directory);
+			}
+		}
+
+		/// <summary>
+		/// Deletes a file, whether it's read-only or not
+		/// </summary>
+		/// <param name="FileName">Name of the file to delete</param>
+		public static void ForceDeleteFile(string FileName)
+		{
+			ForceDeleteFile(new FileInfo(FileName));
+		}
+
+		/// <summary>
+		/// Deletes a file, whether it's read-only or not
+		/// </summary>
+		/// <param name="File">The file to delete</param>
+		public static void ForceDeleteFile(FileInfo File)
+		{
+			try
+			{
+				if (File.Exists)
+				{
+					File.Attributes = FileAttributes.Normal;
+					File.Delete();
+				}
+			}
+			catch (Exception Ex)
+			{
+				throw new WrappedFileOrDirectoryException(Ex, String.Format("Unable to delete '{0}': {1}", File.FullName, Ex.Message.TrimEnd()));
+			}
+		}
+
+		/// <summary>
+		/// Deletes a file, whether it's read-only or not
+		/// </summary>
+		/// <param name="Location">The file to delete</param>
+		public static void ForceDeleteFile(FileReference Location)
+		{
+			ForceDeleteFile(new FileInfo(Location.FullName));
+		}
+
+		/// <summary>
+		/// Deletes a directory and all its contents. Attempts to handle directories with long filenames (> 260 chars) on Windows.
+		/// </summary>
+		/// <param name="DirectoryName">Directory to delete</param>
+		public static void ForceDeleteDirectory(string DirectoryName)
+		{
+			if (Environment.OSVersion.Platform == PlatformID.Win32NT)
+			{
+				ForceDeleteLongDirectoryWin32("\\\\?\\" + DirectoryName);
+			}
+			else
+			{
+				ForceDeleteDirectory(new DirectoryInfo(DirectoryName));
+			}
+		}
+
+		/// <summary>
+		/// Deletes a directory and all its contents. Attempts to handle directories with long filenames (> 260 chars) on Windows.
+		/// </summary>
+		/// <param name="Directory">Directory to delete</param>
+		public static void ForceDeleteDirectory(DirectoryInfo Directory)
+		{
+			if (Directory.Exists)
+			{
+				if (Environment.OSVersion.Platform == PlatformID.Win32NT)
+				{
+					ForceDeleteLongDirectoryWin32("\\\\?\\" + Directory.FullName);
+				}
+				else
+				{
+					ForceDeleteDirectoryContents(Directory);
+					ForceDeleteDirectoryInternal(Directory);
+				}
+			}
+		}
+
+		/// <summary>
+		/// Deletes a directory and all its contents. Attempts to handle directories with long filenames (> 260 chars) on Windows.
+		/// </summary>
+		/// <param name="Location">Directory to delete</param>
+		public static void ForceDeleteDirectory(DirectoryReference Location)
+		{
+			ForceDeleteDirectory(Location.FullName);
+		}
+
+		/// <summary>
+		/// Helper method to delete a directory and throw a WrappedFileOrDirectoryException on failure.
+		/// </summary>
+		/// <param name="Directory">The directory to delete</param>
+		static void ForceDeleteDirectoryInternal(DirectoryInfo Directory)
+		{
+			try
+			{
+				Directory.Delete(true);
+			}
+			catch (Exception Ex)
+			{
+				throw new WrappedFileOrDirectoryException(Ex, String.Format("Unable to delete '{0}': {1}", Directory.FullName, Ex.Message.TrimEnd()));
+			}
+		}
+
+		/// <summary>
+		/// Deletes the contents of a directory, without deleting the directory itself. Attempts to handle directories with long filenames (> 260 chars) on Windows.
+		/// </summary>
+		/// <param name="Directory">Directory to delete</param>
+		public static void ForceDeleteDirectoryContents(string Directory)
+		{
+			ForceDeleteDirectoryContents(new DirectoryInfo(Directory));
+		}
+
+		/// <summary>
+		/// Deletes the contents of a directory, without deleting the directory itself. Attempts to handle directories with long filenames (> 260 chars) on Windows.
+		/// </summary>
+		/// <param name="Directory">Directory to delete</param>
+		public static void ForceDeleteDirectoryContents(DirectoryInfo Directory)
+		{
+			if (Environment.OSVersion.Platform == PlatformID.Win32NT)
+			{
+				ForceDeleteLongDirectoryContentsWin32(Directory.FullName);
+			}
+			else
+			{
+				foreach (FileInfo File in Directory.EnumerateFiles())
+				{
+					ForceDeleteFile(File);
+				}
+				foreach (DirectoryInfo SubDirectory in Directory.EnumerateDirectories())
+				{
+					if (SubDirectory.Attributes.HasFlag(FileAttributes.ReparsePoint))
+					{
+						ForceDeleteDirectoryInternal(SubDirectory);
+					}
+					else
+					{
+						ForceDeleteDirectory(SubDirectory);
+					}
+				}
+			}
+		}
+
+		/// <summary>
+		/// Deletes the contents of a directory, without deleting the directory itself. Attempts to handle directories with long filenames (> 260 chars) on Windows.
+		/// </summary>
+		/// <param name="Directory">Directory to delete</param>
+		public static void ForceDeleteDirectoryContents(DirectoryReference Directory)
+		{
+			ForceDeleteDirectoryContents(new DirectoryInfo(Directory.FullName));
+		}
+
+		/// <summary>
+		/// Moves a file from one location to another. Creates the destination directory, and removes read-only files in the target location if necessary.
+		/// </summary>
+		/// <param name="SourceFileName">Path to the source file</param>
+		/// <param name="TargetFileName">Path to the target file</param>
+		public static void ForceMoveFile(string SourceFileName, string TargetFileName)
+		{
+			ForceMoveFile(new FileReference(SourceFileName), new FileReference(TargetFileName));
+		}
+
+		/// <summary>
+		/// Moves a file from one location to another. Creates the destination directory, and removes read-only files in the target location if necessary.
+		/// </summary>
+		/// <param name="SourceFileName">Path to the source file</param>
+		/// <param name="TargetFileName">Path to the target file</param>
+		public static void ForceMoveFile(FileReference SourceLocation, FileReference TargetLocation)
+		{
+			try
+			{
+				FileReference.Move(SourceLocation, TargetLocation);
+			}
+			catch
+			{
+				try
+				{
+					if (FileReference.Exists(TargetLocation))
+					{
+						FileReference.SetAttributes(TargetLocation, FileAttributes.Normal);
+						FileReference.Delete(TargetLocation);
+					}
+					else if (DirectoryReference.Exists(TargetLocation.Directory))
+					{
+						CreateDirectoryTree(TargetLocation.Directory);
+					}
+				}
+				catch
+				{
+				}
+
+				// Always retry; creating the directory tree may be a race
+				FileReference.Move(SourceLocation, TargetLocation);
+			}
+		}
+
+		#region Win32 Native File Methods
+
+		[StructLayout(LayoutKind.Sequential, CharSet = CharSet.Unicode)]
+		struct WIN32_FIND_DATA
+		{
+			public uint dwFileAttributes;
+			public System.Runtime.InteropServices.ComTypes.FILETIME ftCreationTime;
+			public System.Runtime.InteropServices.ComTypes.FILETIME ftLastAccessTime;
+			public System.Runtime.InteropServices.ComTypes.FILETIME ftLastWriteTime;
+			public uint nFileSizeHigh;
+			public uint nFileSizeLow;
+			public uint dwReserved0;
+			public uint dwReserved1;
+			[MarshalAs(UnmanagedType.ByValTStr, SizeConst = 260)]
+			public string cFileName;
+			[MarshalAs(UnmanagedType.ByValTStr, SizeConst = 14)]
+			public string cAlternateFileName;
+		}
+
+		const uint FILE_ATTRIBUTE_READONLY = 0x01;
+		const uint FILE_ATTRIBUTE_DIRECTORY = 0x10;
+		const uint FILE_ATTRIBUTE_NORMAL = 0x80;
+
+		[DllImport("kernel32.dll", SetLastError = true, CharSet = CharSet.Unicode)]
+		static extern IntPtr FindFirstFileW(string FileName, ref WIN32_FIND_DATA FindData);
+
+		[DllImport("kernel32.dll", SetLastError = true, CharSet = CharSet.Unicode)]
+		[return: MarshalAs(UnmanagedType.Bool)]
+		static extern bool FindNextFileW(IntPtr FindHandle, ref WIN32_FIND_DATA FindData);
+
+		[DllImport("kernel32.dll", SetLastError = true, CharSet = CharSet.Unicode)]
+		[return: MarshalAs(UnmanagedType.Bool)]
+		static extern bool FindClose(IntPtr FindHandle);
+
+		[DllImport("kernel32.dll", SetLastError = true, CharSet = CharSet.Unicode)]
+		[return: MarshalAs(UnmanagedType.Bool)]
+		static extern bool DeleteFileW(string lpFileName);
+
+		[DllImport("kernel32.dll", SetLastError = true, CharSet = CharSet.Unicode)]
+		[return: MarshalAs(UnmanagedType.Bool)]
+		static extern bool RemoveDirectory(string lpPathName);
+
+		[DllImport("kernel32.dll", SetLastError = true, CharSet = CharSet.Unicode)]
+		[return: MarshalAs(UnmanagedType.Bool)]
+		static extern bool SetFileAttributesW(string lpFileName, uint dwFileAttributes);
+
+		static readonly IntPtr INVALID_HANDLE_VALUE = new IntPtr(-1);
+
+		private static void ForceDeleteLongDirectoryContentsWin32(string DirName)
+		{
+			WIN32_FIND_DATA FindData = new WIN32_FIND_DATA();
+
+			IntPtr hFind = FindFirstFileW(DirName + "\\*", ref FindData);
+			if (hFind != INVALID_HANDLE_VALUE)
+			{
+				try
+				{
+					for (; ; )
+					{
+						string FullName = DirName + "\\" + FindData.cFileName;
+						if ((FindData.dwFileAttributes & FILE_ATTRIBUTE_DIRECTORY) != 0)
+						{
+							if (FindData.cFileName != "." && FindData.cFileName != "..")
+							{
+								ForceDeleteLongDirectoryWin32(FullName);
+							}
+						}
+						else
+						{
+							if ((FindData.dwFileAttributes & FILE_ATTRIBUTE_READONLY) != 0)
+							{
+								SetFileAttributesW(FullName, FILE_ATTRIBUTE_NORMAL);
+							}
+							if (!DeleteFileW(FullName))
+							{
+								throw new WrappedFileOrDirectoryException(new Win32Exception(), "Unable to delete " + FullName);
+							}
+						}
+
+						if (!FindNextFileW(hFind, ref FindData))
+						{
+							break;
+						}
+					}
+				}
+				finally
+				{
+					FindClose(hFind);
+				}
+			}
+		}
+
+		private static void ForceDeleteLongDirectoryWin32(string DirName)
+		{
+			ForceDeleteLongDirectoryContentsWin32(DirName);
+
+			if (!RemoveDirectory(DirName))
+			{
+				throw new WrappedFileOrDirectoryException(new Win32Exception(), "Unable to delete " + DirName);
+			}
+		}
+
+		#endregion
 	}
 }
diff --git a/Engine/Source/Programs/DotNETCommon/DotNETUtilities/HelpUtils.cs b/Engine/Source/Programs/DotNETCommon/DotNETUtilities/HelpUtils.cs
--- a/Engine/Source/Programs/DotNETCommon/DotNETUtilities/HelpUtils.cs
+++ b/Engine/Source/Programs/DotNETCommon/DotNETUtilities/HelpUtils.cs
@@ -1,3 +1,5 @@
+// Copyright 1998-2019 Epic Games, Inc. All Rights Reserved.
+
 using System;
 using System.Collections.Generic;
 using System.ComponentModel;
@@ -45,62 +47,46 @@ namespace Tools.DotNETCommon
 		/// <param name="Type">Type to print help for</param>
 		public static void PrintHelp(string Title, Type Type)
 		{
-			List<string> Lines = GetHelpText(Title, Type);
-			foreach(string Line in Lines)
-			{
-				Log.TraceInformation("{0}", Line);
-			}
+			PrintHelp(Title, GetDescription(Type), CommandLineArguments.GetParameters(Type));
 		}
 
 		/// <summary>
-		/// Get the help text for a given type, using any attribute-driven command line arguments, and the [Description] attribute for help text.
+		/// Prints help for a command
 		/// </summary>
 		/// <param name="Title">Title for the help text</param>
-		/// <param name="Type">Type to query</param>
-		/// <param name="MaxWidth">Maximum width for the returned lines</param>
-		/// <returns>List of help lines</returns>
-		public static List<string> GetHelpText(string Title, Type Type)
+		/// <param name="Description">Description for the command</param>
+		/// <param name="Parameters">List of parameters</param>
+		public static void PrintHelp(string Title, string Description, List<KeyValuePair<string, string>> Parameters)
 		{
-			return GetHelpText(Title, GetDescription(Type), CommandLineArguments.GetParameters(Type), WindowWidth - 1);
-		}
-
-		/// <summary>
-		/// Get the help text for a given type, using any attribute-driven command line arguments, and the [Description] attribute for help text.
-		/// </summary>
-		/// <param name="Title">Title for the help text</param>
-		/// <param name="Description">Help description</param>
-		/// <param name="Parameters">Arguments for the command</param>
-		/// <param name="MaxWidth">Maximum width for the returned lines</param>
-		/// <returns>List of help lines</returns>
-		public static List<string> GetHelpText(string Title, string Description, List<KeyValuePair<string, string>> Parameters, int MaxWidth)
-		{
-			List<string> Lines = new List<string>();
+			int MaxWidth = WindowWidth;
 
+			bool bFirstLine = true;
 			if (!String.IsNullOrEmpty(Title))
 			{
-				Lines.AddRange(StringUtils.WordWrap(Title, MaxWidth));
+				PrintParagraph(Title, MaxWidth);
+				bFirstLine = false;
 			}
 
 			if (!String.IsNullOrEmpty(Description))
 			{
-				if (Lines.Count > 0)
+				if (!bFirstLine)
 				{
-					Lines.Add("");
+					Log.TraceInformation("");
 				}
-				Lines.AddRange(StringUtils.WordWrap(Description, MaxWidth));
+				PrintParagraph(Description, MaxWidth);
+				bFirstLine = false;
 			}
 
 			if (Parameters.Count > 0)
 			{
-				if (Lines.Count > 0)
+				if (!bFirstLine)
 				{
-					Lines.Add("");
+					Log.TraceInformation("");
 				}
-				Lines.Add("Parameters:");
-				Lines.AddRange(Tabulate(Parameters, 4, 24, MaxWidth));
-			}
 
-			return Lines;
+				Log.TraceInformation("Parameters:");
+				PrintTable(Parameters, 4, 24, MaxWidth);
+			}
 		}
 
 		/// <summary>
@@ -122,6 +108,48 @@ namespace Tools.DotNETCommon
 			return DescriptionText.ToString();
 		}
 
+		/// <summary>
+		/// Prints a paragraph of text using word wrapping
+		/// </summary>
+		/// <param name="Text">Text to print</param>
+		/// <param name="MaxWidth">Maximum width for each line</param>
+		public static void PrintParagraph(string Text)
+		{
+			PrintParagraph(Text, WindowWidth - 1);
+		}
+
+		/// <summary>
+		/// Prints a paragraph of text using word wrapping
+		/// </summary>
+		/// <param name="Text">Text to print</param>
+		/// <param name="MaxWidth">Maximum width for each line</param>
+		public static void PrintParagraph(string Text, int MaxWidth)
+		{
+			List<string> Lines = StringUtils.WordWrap(Text, MaxWidth);
+			foreach (string Line in Lines)
+			{
+				Log.TraceInformation("{0}", Line);
+			}
+		}
+
+		/// <summary>
+		/// Formats the given parameters as so:
+		///     -Param1     Param1 Description
+		///
+		///     -Param2      Param2 Description, this description is
+		///                  longer and splits onto a separate line. 
+		///
+		///     -Param3      Param3 Description continues as before. 
+		/// </summary>
+		/// <param name="Items">List of parameters arranged as "-ParamName Param Description"</param>
+		/// <param name="Indent">Indent from the left hand side</param>
+		/// <param name="DefaultRightPadding">The minimum padding from the start of the param name to the start of the description (resizes with larger param names)</param>
+		/// <returns></returns>
+		public static void PrintTable(List<KeyValuePair<string, string>> Items, int Indent, int MinFirstColumnWidth)
+		{
+			PrintTable(Items, Indent, MinFirstColumnWidth, WindowWidth - 1);
+		}
+
 		/// <summary>
 		/// Formats the given parameters as so:
 		///     -Param1     Param1 Description
@@ -135,9 +163,8 @@ namespace Tools.DotNETCommon
 		/// <param name="Indent">Indent from the left hand side</param>
 		/// <param name="DefaultRightPadding">The minimum padding from the start of the param name to the start of the description (resizes with larger param names)</param>
 		/// <returns></returns>
-		public static List<string> Tabulate(List<KeyValuePair<string, string>> Items, int Indent, int MinFirstColumnWidth, int MaxWidth)
+		public static void PrintTable(List<KeyValuePair<string, string>> Items, int Indent, int MinFirstColumnWidth, int MaxWidth)
 		{
-			List<string> Lines = new List<string>();
 			if(Items.Count > 0)
 			{
 				// string used to intent the param
@@ -159,14 +186,13 @@ namespace Tools.DotNETCommon
 					{
 						// Formatting as following:
 						// <Indent>-param<Right Padding>Description<New line>
-						Lines.Add(ParamString + DescriptionLine);
+						Log.TraceInformation("{0}{1}", ParamString, DescriptionLine);
 
 						// we replace the param string on subsequent lines with white space of the same length
 						ParamString = string.Empty.PadRight(IndentString.Length + RightPadding);
 					}
 				}
 			}
-			return Lines;
 		}
 	}
 }
diff --git a/Engine/Source/Programs/DotNETCommon/DotNETUtilities/ManagedProcess.cs b/Engine/Source/Programs/DotNETCommon/DotNETUtilities/ManagedProcess.cs
--- a/Engine/Source/Programs/DotNETCommon/DotNETUtilities/ManagedProcess.cs
+++ b/Engine/Source/Programs/DotNETCommon/DotNETUtilities/ManagedProcess.cs
@@ -289,18 +289,18 @@ namespace Tools.DotNETCommon
 		/// <param name="Input">Text to be passed via stdin to the new process. May be null.</param>
 		public ManagedProcess(ManagedProcessGroup Group, string FileName, string CommandLine, string WorkingDirectory, IReadOnlyDictionary<string, string> Environment, byte[] Input, ProcessPriorityClass Priority)
 		{
-			if(ManagedProcessGroup.SupportsJobObjects)
+			if (ManagedProcessGroup.SupportsJobObjects)
 			{
 				// Create the child process
 				IntPtr EnvironmentBlock = IntPtr.Zero;
 				try
 				{
 					// Create the environment block for the child process, if necessary.
-					if(Environment != null)
+					if (Environment != null)
 					{
 						// The native format for the environment block is a sequence of null terminated strings with a final null terminator.
 						List<byte> EnvironmentBytes = new List<byte>();
-						foreach(KeyValuePair<string, string> Pair in Environment)
+						foreach (KeyValuePair<string, string> Pair in Environment)
 						{
 							EnvironmentBytes.AddRange(Console.OutputEncoding.GetBytes(Pair.Key));
 							EnvironmentBytes.Add((byte)'=');
@@ -319,7 +319,7 @@ namespace Tools.DotNETCommon
 					{
 						// Get the flags to create the new process
 						ProcessCreationFlags Flags = ProcessCreationFlags.CREATE_NO_WINDOW | ProcessCreationFlags.CREATE_SUSPENDED;
-						switch(Priority)
+						switch (Priority)
 						{
 							case ProcessPriorityClass.Normal:
 								Flags |= ProcessCreationFlags.NORMAL_PRIORITY_CLASS;
@@ -344,7 +344,7 @@ namespace Tools.DotNETCommon
 						// Acquire a global lock before creating inheritable handles. If multiple threads create inheritable handles at the same time, and child processes will inherit them all.
 						// Since we need to wait for output pipes to be closed (in order to consume all output), this can result in output reads not returning until all processes with the same
 						// inherited handles are closed.
-						lock(LockObject)
+						lock (LockObject)
 						{
 							SafeFileHandle StdInRead = null;
 							SafeFileHandle StdOutWrite = null;
@@ -356,15 +356,15 @@ namespace Tools.DotNETCommon
 								SecurityAttributes.nLength = Marshal.SizeOf(SecurityAttributes);
 								SecurityAttributes.bInheritHandle = 1;
 
-								if(CreatePipe(out StdInRead, out StdInWrite, SecurityAttributes, 0) == 0 || SetHandleInformation(StdInWrite, HANDLE_FLAG_INHERIT, 0) == 0)
+								if (CreatePipe(out StdInRead, out StdInWrite, SecurityAttributes, 0) == 0 || SetHandleInformation(StdInWrite, HANDLE_FLAG_INHERIT, 0) == 0)
 								{
 									throw new Win32Exception();
 								}
-								if(CreatePipe(out StdOutRead, out StdOutWrite, SecurityAttributes, 0) == 0 || SetHandleInformation(StdOutRead, HANDLE_FLAG_INHERIT, 0) == 0)
+								if (CreatePipe(out StdOutRead, out StdOutWrite, SecurityAttributes, 0) == 0 || SetHandleInformation(StdOutRead, HANDLE_FLAG_INHERIT, 0) == 0)
 								{
 									throw new Win32Exception();
 								}
-								if(DuplicateHandle(GetCurrentProcess(), StdOutWrite, GetCurrentProcess(), out StdErrWrite, 0, true, DUPLICATE_SAME_ACCESS) == 0)
+								if (DuplicateHandle(GetCurrentProcess(), StdOutWrite, GetCurrentProcess(), out StdErrWrite, 0, true, DUPLICATE_SAME_ACCESS) == 0)
 								{
 									throw new Win32Exception(String.Format("Unable to duplicate stdout handle ({0})", StdOutWrite));
 								}
@@ -385,17 +385,17 @@ namespace Tools.DotNETCommon
 							finally
 							{
 								// Close the write ends of the handle. We don't want any other process to be able to inherit these.
-								if(StdInRead != null)
+								if (StdInRead != null)
 								{
 									StdInRead.Dispose();
 									StdInRead = null;
 								}
-								if(StdOutWrite != null)
+								if (StdOutWrite != null)
 								{
 									StdOutWrite.Dispose();
 									StdOutWrite = null;
 								}
-								if(StdErrWrite != null)
+								if (StdErrWrite != null)
 								{
 									StdErrWrite.Dispose();
 									StdErrWrite = null;
@@ -404,7 +404,7 @@ namespace Tools.DotNETCommon
 						}
 
 						// Add it to our job object
-						if(AssignProcessToJobObject(Group.JobHandle, ProcessInfo.hProcess) == 0)
+						if (Group != null && AssignProcessToJobObject(Group.JobHandle, ProcessInfo.hProcess) == 0)
 						{
 							// Support for nested job objects was only addeed in Windows 8; prior to that, assigning processes to job objects would fail. Figure out if we're already in a job, and ignore the error if we are.
 							int OriginalError = Marshal.GetLastWin32Error();
@@ -412,22 +412,22 @@ namespace Tools.DotNETCommon
 							bool bProcessInJob;
 							IsProcessInJob(GetCurrentProcess(), IntPtr.Zero, out bProcessInJob);
 
-							if(!bProcessInJob)
+							if (!bProcessInJob)
 							{
 								throw new Win32Exception(OriginalError);
 							}
 						}
 
 						// Allow the thread to start running
-						if(ResumeThread(ProcessInfo.hThread) == -1)
+						if (ResumeThread(ProcessInfo.hThread) == -1)
 						{
 							throw new Win32Exception();
 						}
 
 						// If we have any input text, write it to stdin now
-						using(FileStream Stream = new FileStream(StdInWrite, FileAccess.Write, 4096, false))
+						using (FileStream Stream = new FileStream(StdInWrite, FileAccess.Write, 4096, false))
 						{
-							if(Input != null)
+							if (Input != null)
 							{
 								Stream.Write(Input, 0, Input.Length);
 								Stream.Flush();
@@ -443,11 +443,11 @@ namespace Tools.DotNETCommon
 					}
 					finally
 					{
-						if(ProcessInfo.hProcess != IntPtr.Zero && ProcessHandle == null)
+						if (ProcessInfo.hProcess != IntPtr.Zero && ProcessHandle == null)
 						{
 							CloseHandle(ProcessInfo.hProcess);
 						}
-						if(ProcessInfo.hThread != IntPtr.Zero)
+						if (ProcessInfo.hThread != IntPtr.Zero)
 						{
 							CloseHandle(ProcessInfo.hThread);
 						}
@@ -460,7 +460,7 @@ namespace Tools.DotNETCommon
 				}
 				finally
 				{
-					if(EnvironmentBlock != IntPtr.Zero)
+					if (EnvironmentBlock != IntPtr.Zero)
 					{
 						Marshal.FreeHGlobal(EnvironmentBlock);
 						EnvironmentBlock = IntPtr.Zero;
@@ -479,9 +479,9 @@ namespace Tools.DotNETCommon
 				FrameworkProcess.StartInfo.UseShellExecute = false;
 				FrameworkProcess.StartInfo.CreateNoWindow = true;
 
-				if(Environment != null)
+				if (Environment != null)
 				{
-					foreach(KeyValuePair<string, string> Pair in Environment)
+					foreach (KeyValuePair<string, string> Pair in Environment)
 					{
 						FrameworkProcess.StartInfo.EnvironmentVariables[Pair.Key] = Pair.Value;
 					}
diff --git a/Engine/Source/Programs/DotNETCommon/DotNETUtilities/ProcessUtils.cs b/Engine/Source/Programs/DotNETCommon/DotNETUtilities/ProcessUtils.cs
new file mode 100644
--- /dev/null
+++ b/Engine/Source/Programs/DotNETCommon/DotNETUtilities/ProcessUtils.cs
@@ -0,0 +1,197 @@
+// Copyright 1998-2019 Epic Games, Inc. All Rights Reserved.
+
+using Microsoft.Win32.SafeHandles;
+using System;
+using System.Collections.Generic;
+using System.Diagnostics;
+using System.Linq;
+using System.Runtime.InteropServices;
+using System.Text;
+using System.Threading;
+using System.Threading.Tasks;
+
+namespace Tools.DotNETCommon
+{
+	/// <summary>
+	/// Includes utilities for managing processes
+	/// </summary>
+	public static class ProcessUtils
+	{
+		[StructLayout(LayoutKind.Sequential)]
+		struct PROCESSENTRY32
+		{
+			public uint dwSize;
+			public uint cntUsage;
+			public uint th32ProcessID;
+			public IntPtr th32DefaultHeapID;
+			public uint th32ModuleID;
+			public uint cntThreads;
+			public uint th32ParentProcessID;
+			public int pcPriClassBase;
+			public uint dwFlags;
+			[MarshalAs(UnmanagedType.ByValTStr, SizeConst = 260)]
+			public string szExeFile;
+		};
+
+		static uint TH32CS_SNAPPROCESS = 2;
+
+		[DllImport("kernel32.dll", SetLastError = true)]
+		static extern SafeFileHandle CreateToolhelp32Snapshot(uint dwFlags, uint th32ProcessID);
+
+		[DllImport("kernel32.dll")]
+		[return: MarshalAs(UnmanagedType.Bool)]
+		static extern bool Process32First(SafeFileHandle hSnapshot, ref PROCESSENTRY32 lppe);
+
+		[DllImport("kernel32.dll")]
+		[return: MarshalAs(UnmanagedType.Bool)]
+		static extern bool Process32Next(SafeFileHandle hSnapshot, ref PROCESSENTRY32 lppe);
+
+		const uint PROCESS_TERMINATE = 0x0001;
+
+		[DllImport("kernel32.dll", SetLastError = true)]
+		static extern SafeProcessHandle OpenProcess(uint dwDesiredAccess, bool bInheritHandle, uint processId);
+
+		[DllImport("kernel32.dll", SetLastError = true)]
+		[return: MarshalAs(UnmanagedType.Bool)]
+		static extern bool TerminateProcess(SafeProcessHandle hProcess, uint uExitCode);
+
+		/// <summary>
+		/// Terminates all child processes of the current process
+		/// </summary>
+		public static void TerminateChildProcesses()
+		{
+			using (Process CurrentProcess = Process.GetCurrentProcess())
+			{
+				if (Environment.OSVersion.Platform == PlatformID.Win32NT)
+				{
+					TerminateChildProcessesWin32(CurrentProcess.Id);
+				}
+				else
+				{
+					TerminateChildProcessesPosix(CurrentProcess.Id);
+				}
+			}
+		}
+
+		/// <summary>
+		/// Win32-specific implementation of TerminateChildProcesses
+		/// </summary>
+		/// <param name="RootProcessId">The root process id of the tree to be terminated</param>
+		static void TerminateChildProcessesWin32(int RootProcessId)
+		{
+			Stopwatch Timer = Stopwatch.StartNew();
+			HashSet<uint> UnableToOpenProcessIds = new HashSet<uint>();
+			for (; ; )
+			{
+				// Find a map of process id to parent process id
+				Dictionary<uint, uint> ProcessIdToParentProcessId = new Dictionary<uint, uint>();
+
+				using (SafeFileHandle SnapshotHandle = CreateToolhelp32Snapshot(TH32CS_SNAPPROCESS, 0))
+				{
+					if (!SnapshotHandle.IsInvalid)
+					{
+						PROCESSENTRY32 ProcessEntry = new PROCESSENTRY32();
+						ProcessEntry.dwSize = (uint)Marshal.SizeOf(typeof(PROCESSENTRY32));
+
+						for (bool bResult = Process32First(SnapshotHandle, ref ProcessEntry); bResult; bResult = Process32Next(SnapshotHandle, ref ProcessEntry))
+						{
+							if (ProcessEntry.th32ParentProcessID != 0)
+							{
+								ProcessIdToParentProcessId[ProcessEntry.th32ProcessID] = ProcessEntry.th32ParentProcessID;
+							}
+						}
+					}
+				}
+
+				// Find any process ids which are a descendant from the root process it
+				List<uint> ChildProcessIds = new List<uint>();
+				foreach (KeyValuePair<uint, uint> Pair in ProcessIdToParentProcessId)
+				{
+					uint ParentProcessId = Pair.Value;
+					for (; ; )
+					{
+						if (ParentProcessId == RootProcessId)
+						{
+							ChildProcessIds.Add(Pair.Key);
+							break;
+						}
+						if (!ProcessIdToParentProcessId.TryGetValue(ParentProcessId, out ParentProcessId))
+						{
+							break;
+						}
+					}
+				}
+
+				// If there's nothing running, we can quit
+				if (ChildProcessIds.Count == 0)
+				{
+					break;
+				}
+
+				// Check if we need to print a message about what's going on
+				bool bPrintStatus = false;
+				if (Timer == null)
+				{
+					Timer = Stopwatch.StartNew();
+					bPrintStatus = true;
+				}
+				else if (Timer.Elapsed > TimeSpan.FromSeconds(10.0))
+				{
+					Timer.Restart();
+					bPrintStatus = true;
+				}
+
+				// Try to kill all the child processes
+				bool bWaitingForProcess = false;
+				foreach (uint ChildProcessId in ChildProcessIds)
+				{
+					using (SafeProcessHandle ProcessHandle = OpenProcess(PROCESS_TERMINATE, false, ChildProcessId))
+					{
+						if (ProcessHandle.IsInvalid)
+						{
+							if (UnableToOpenProcessIds.Add(ChildProcessId))
+							{
+								bWaitingForProcess = true;
+							}
+							else
+							{
+								Log.TraceWarning("Unable to terminate process {0}.", ChildProcessId);
+							}
+						}
+						else
+						{
+							if (bPrintStatus)
+							{
+								Log.TraceInformation("Waiting for process {0} to terminate.", ChildProcessId);
+							}
+
+							TerminateProcess(ProcessHandle, 3);
+							bWaitingForProcess = true;
+						}
+					}
+				}
+
+				// If there wasn't anything we could do, bail out
+				if(!bWaitingForProcess)
+				{
+					break;
+				}
+
+				// Allow a short amount of time for processes to exit, then check again
+				Thread.Sleep(500);
+			}
+		}
+
+		/// <summary>
+		/// Posix implementation of TerminateChildProcesses
+		/// </summary>
+		/// <param name="RootProcessId">The root process id of the tree to be terminated</param>
+		static void TerminateChildProcessesPosix(int RootProcessId)
+		{
+			using (Process Process = Process.Start("pkill", String.Format("-9 -p {0}", RootProcessId)))
+			{
+				Process.WaitForExit();
+			}
+		}
+	}
+}
diff --git a/Engine/Source/Programs/DotNETCommon/DotNETUtilities/StringUtils.cs b/Engine/Source/Programs/DotNETCommon/DotNETUtilities/StringUtils.cs
--- a/Engine/Source/Programs/DotNETCommon/DotNETUtilities/StringUtils.cs
+++ b/Engine/Source/Programs/DotNETCommon/DotNETUtilities/StringUtils.cs
@@ -1,4 +1,5 @@
 ﻿// Copyright 1998-2019 Epic Games, Inc. All Rights Reserved.
+
 using System;
 using System.Collections.Generic;
 using System.Linq;
@@ -164,6 +165,20 @@ namespace Tools.DotNETCommon
 			}
 		}
 
+		/// <summary>
+		/// Parses a hexadecimal string into an array of bytes
+		/// </summary>
+		/// <returns>Array of bytes</returns>
+		public static byte[] ParseHexString(string Text)
+		{
+			byte[] Bytes;
+			if(!TryParseHexString(Text, out Bytes))
+			{
+				throw new FormatException(String.Format("Invalid hex string: '{0}'", Text));
+			}
+			return Bytes;
+		}
+
 		/// <summary>
 		/// Parses a hexadecimal string into an array of bytes
 		/// </summary>
@@ -171,8 +186,13 @@ namespace Tools.DotNETCommon
 		/// <returns></returns>
 		public static bool TryParseHexString(string Text, out byte[] OutBytes)
 		{
+			if((Text.Length & 1) != 0)
+			{
+				throw new FormatException("Length of hex string must be a multiple of two characters");
+			}
+
 			byte[] Bytes = new byte[Text.Length / 2];
-			for(int Idx = 0; Idx < Text.Length; Idx++)
+			for(int Idx = 0; Idx < Text.Length; Idx += 2)
 			{
 				int A, B;
 				if(!TryParseHexDigit(Text[Idx], out A) || !TryParseHexDigit(Text[Idx + 1], out B))
