Updates to BuildAgent executable. Fix a few issues with the EC post-processor, and make it easier to pass command line arguments to the child process for a run command.

#rnx

diff --git a/Engine/Source/Programs/BuildAgent/Program.cs b/Engine/Source/Programs/BuildAgent/Program.cs
--- a/Engine/Source/Programs/BuildAgent/Program.cs
+++ b/Engine/Source/Programs/BuildAgent/Program.cs
@@ -48,10 +48,10 @@ namespace BuildAgent
 		static int GuardedMain(string[] Args)
 		{
 			// Find the index of the first command
-			int FirstModeIndex = 0;
-			while (FirstModeIndex < Args.Length && Args[FirstModeIndex].StartsWith("-"))
+			int ModeIndex = 0;
+			while (ModeIndex < Args.Length && Args[ModeIndex].StartsWith("-"))
 			{
-				FirstModeIndex++;
+				ModeIndex++;
 			}
 
 			// Find all the ToolMode types
@@ -66,14 +66,14 @@ namespace BuildAgent
 			}
 
 			// Check if there are any commands specified on the command line.
-			if (FirstModeIndex == Args.Length)
+			if (ModeIndex == Args.Length)
 			{
 				Log.TraceInformation("BuildAgent");
 				Log.TraceInformation("");
 				Log.TraceInformation("Utility for managing automated processes on build machines.");
 				Log.TraceInformation("");
 				Log.TraceInformation("Usage:");
-				Log.TraceInformation("    BuildAgent.exe [CommandA] [-OptionA1] [-OptionA2] [CommandB] [-OptionB1]...");
+				Log.TraceInformation("    BuildAgent.exe [Command] [-Option1] [-Option2]...");
 				Log.TraceInformation("");
 				Log.TraceInformation("Commands:");
 
@@ -81,69 +81,59 @@ namespace BuildAgent
 				return 0;
 			}
 
-			// Parse the arguments for all commands
-			List<ProgramMode> Modes = new List<ProgramMode>();
-			for (int NextCommandIdx = FirstModeIndex; NextCommandIdx < Args.Length;)
-			{
-				string CommandName = Args[NextCommandIdx++];
+			// Get the command name
+			string CommandName = Args[ModeIndex];
 
-				// Get the command type
-				Type CommandType;
-				if (!ModeNameToType.TryGetValue(CommandName, out CommandType))
-				{
-					Log.TraceError("Unknown command '{0}'.", CommandName);
-					Log.TraceInformation("");
-					Log.TraceInformation("Available commands");
-					PrintCommands(ModeNameToType);
-					return 1;
-				}
+			// Get the command type
+			Type CommandType;
+			if (!ModeNameToType.TryGetValue(CommandName, out CommandType))
+			{
+				Log.TraceError("Unknown command '{0}'.", CommandName);
+				Log.TraceInformation("");
+				Log.TraceInformation("Available commands");
+				PrintCommands(ModeNameToType);
+				return 1;
+			}
 
-				// Build an argument list for the command, including all the global arguments as well as arguments until the next command
-				List<string> ArgumentList = new List<string>();
-				for (int Idx = 0; Idx < FirstModeIndex; Idx++)
+			// Build an argument list for the command, including all the global arguments as well as arguments until the next command
+			List<string> ArgumentList = new List<string>();
+			for (int Idx = 0; Idx < Args.Length; Idx++)
+			{
+				if (Idx != ModeIndex)
 				{
 					ArgumentList.Add(Args[Idx]);
 				}
-				while (NextCommandIdx < Args.Length && Args[NextCommandIdx].StartsWith("-"))
-				{
-					ArgumentList.Add(Args[NextCommandIdx++]);
-				}
-				CommandLineArguments CommandArguments = new CommandLineArguments(ArgumentList.ToArray());
-
-				// Create the command instance
-				ProgramMode Command = (ProgramMode)Activator.CreateInstance(CommandType);
-				Modes.Add(Command);
-
-				// If the help flag is specified, print the help info and exit immediately
-				if (CommandArguments.HasOption("-Help"))
-				{
-					HelpUtils.PrintHelp(CommandName, HelpUtils.GetDescription(CommandType), Command.GetParameters());
-					return 1;
-				}
+			}
+			CommandLineArguments ModeArguments = new CommandLineArguments(ArgumentList.ToArray());
 
-				// Configure the command
-				try
-				{
-					Command.Configure(CommandArguments);
-					CommandArguments.CheckAllArgumentsUsed();
-				}
-				catch (CommandLineArgumentException Ex)
-				{
-					Log.TraceError("{0}: {1}", CommandName, Ex.Message);
-					Log.TraceInformation("");
-					Log.TraceInformation("Arguments for {0}:", CommandName);
+			// Create the command instance
+			ProgramMode Mode = (ProgramMode)Activator.CreateInstance(CommandType);
 
-					HelpUtils.PrintTable(Command.GetParameters(), 4, 24);
-					return 1;
-				}
+			// If the help flag is specified, print the help info and exit immediately
+			if (ModeArguments.HasOption("-Help"))
+			{
+				HelpUtils.PrintHelp(CommandName, HelpUtils.GetDescription(CommandType), Mode.GetParameters());
+				return 1;
 			}
 
-			// Execute all the commands
-			foreach (ProgramMode Mode in Modes)
+			// Configure the command
+			try
 			{
-				Mode.Execute();
+				Mode.Configure(ModeArguments);
+				ModeArguments.CheckAllArgumentsUsed();
 			}
+			catch (CommandLineArgumentException Ex)
+			{
+				Log.TraceError("{0}: {1}", CommandName, Ex.Message);
+				Log.TraceInformation("");
+				Log.TraceInformation("Arguments for {0}:", CommandName);
 
+				HelpUtils.PrintTable(Mode.GetParameters(), 4, 24);
+				return 1;
+			}
+
+			// Execute all the commands
+			Mode.Execute();
 			return 0;
 		}
 
diff --git a/Engine/Source/Programs/BuildAgent/Run/Listeners/ElectricCommanderListener.cs b/Engine/Source/Programs/BuildAgent/Run/Listeners/ElectricCommanderListener.cs
--- a/Engine/Source/Programs/BuildAgent/Run/Listeners/ElectricCommanderListener.cs
+++ b/Engine/Source/Programs/BuildAgent/Run/Listeners/ElectricCommanderListener.cs
@@ -8,6 +8,7 @@ using System.Diagnostics;
 using System.IO;
 using System.Linq;
 using System.Text;
+using System.Threading;
 using System.Threading.Tasks;
 using System.Xml;
 using Tools.DotNETCommon;
@@ -25,19 +26,34 @@ namespace BuildAgent.Listeners
 		FileReference OutputFile;
 
 		/// <summary>
-		/// List of the matched errors
+		/// New errors which are posted for the background thread
 		/// </summary>
-		List<ErrorMatch> Errors = new List<ErrorMatch>();
+		List<ErrorMatch> NewErrors = new List<ErrorMatch>();
+
+		/// <summary>
+		/// Object used for synchronization
+		/// </summary>
+		object LockObject = new object();
 
 		/// <summary>
 		/// Whether we should set a property indicating that the diagnostics file exists
 		/// </summary>
-		bool bUpdateProperty = true;
+		bool bUpdateProperties;
+
+		/// <summary>
+		/// Thread used to update EC properties with the error/warning count
+		/// </summary>
+		Thread BackgroundThread;
+
+		/// <summary>
+		/// Event that is set whenever EC needs to be updated
+		/// </summary>
+		ManualResetEvent UpdateEvent = new ManualResetEvent(false);
 
 		/// <summary>
-		/// Process which is created to update properties in the background
+		/// Set when the object is being disposed
 		/// </summary>
-		Process UpdatePropertyProcess;
+		bool bDisposing;
 
 		/// <summary>
 		/// Constructs an EC listener
@@ -45,16 +61,16 @@ namespace BuildAgent.Listeners
 		/// <param name="OutputFile">Path to the output file</param>
 		public ElectricCommanderListener()
 		{
-			string EnvVar = Environment.GetEnvironmentVariable("COMMANDER_JOBSTEP");
+			string EnvVar = Environment.GetEnvironmentVariable("COMMANDER_JOBSTEPID");
 			if (String.IsNullOrEmpty(EnvVar))
 			{
 				OutputFile = new FileReference("BuildAgent.xml");
-				bUpdateProperty = false;
+				bUpdateProperties = false;
 			}
 			else
 			{
 				OutputFile = new FileReference(String.Format("BuildAgent-{0}.xml", EnvVar));
-				bUpdateProperty = true;
+				bUpdateProperties = true;
 			}
 		}
 
@@ -63,19 +79,21 @@ namespace BuildAgent.Listeners
 		/// </summary>
 		public void Dispose()
 		{
-			if(UpdatePropertyProcess != null)
+			if (BackgroundThread != null)
 			{
-				if (!UpdatePropertyProcess.WaitForExit(0))
+				lock(LockObject)
 				{
-					Log.TraceInformation("Waiting for ectool to terminate...");
-					UpdatePropertyProcess.WaitForExit();
+					bDisposing = true;
 				}
-				if (UpdatePropertyProcess.ExitCode != 0)
+				UpdateEvent.Set();
+
+				if (!BackgroundThread.Join(500))
 				{
-					Log.TraceWarning("ectool terminated with exit code {0}", UpdatePropertyProcess.ExitCode);
+					Log.TraceInformation("Waiting for ectool to terminate...");
+					BackgroundThread.Join();
 				}
-				UpdatePropertyProcess.Dispose();
-				UpdatePropertyProcess = null;
+
+				BackgroundThread = null;
 			}
 		}
 
@@ -85,21 +103,111 @@ namespace BuildAgent.Listeners
 		/// <param name="Error">The matched error</param>
 		public void OnErrorMatch(ErrorMatch Error)
 		{
-			Errors.Add(Error);
+			lock(LockObject)
+			{
+				NewErrors.Add(Error);
+			}
+
+			UpdateEvent.Set();
 
-			Write();
+			if (BackgroundThread == null)
+			{
+				BackgroundThread = new Thread(() => BackgroundWorker());
+				BackgroundThread.Start();
+			}
+		}
+
+		/// <summary>
+		/// Updates EC properties in the background
+		/// </summary>
+		private void BackgroundWorker()
+		{
+			bool bHasSetDiagFile = false;
+			int NumErrors = 0;
+			int NumWarnings = 0;
+			string Outcome = null;
 
-			if(bUpdateProperty)
+			List<ErrorMatch> Errors = new List<ErrorMatch>();
+			while (!bDisposing)
 			{
-				UpdatePropertyProcess = Process.Start("ectool", String.Format("setProperty /myJobStep/diagFile {0}", OutputFile.GetFileName()));
-				bUpdateProperty = false;
+				// Copy the current set of errors
+				lock(LockObject)
+				{
+					if (NewErrors.Count > 0)
+					{
+						Errors.AddRange(NewErrors);
+						NewErrors.Clear();
+					}
+					else if(bDisposing)
+					{
+						break;
+					}
+				}
+
+				// Write them to disk
+				Write(OutputFile, Errors);
+
+				// On the first run, set the path to the diagnostics file
+				if (!bHasSetDiagFile)
+				{
+					SetProperty("/myJobStep/diagFile", OutputFile.GetFileName());
+					bHasSetDiagFile = true;
+				}
+
+				// Check if the number of errors has changed
+				int NewNumErrors = Errors.Count(x => x.Severity == ErrorSeverity.Error);
+				if (NewNumErrors > NumErrors)
+				{
+					SetProperty("/myJobStep/errors", NewNumErrors.ToString());
+					NumErrors = NewNumErrors;
+				}
+
+				// Check if the number of warnings has changed
+				int NewNumWarnings = Errors.Count(x => x.Severity == ErrorSeverity.Warning);
+				if (NewNumWarnings > NumWarnings)
+				{
+					SetProperty("/myJobStep/warnings", NewNumWarnings.ToString());
+					NumWarnings = NewNumWarnings;
+				}
+
+				// Check if the outcome has changed
+				string NewOutcome = (NumErrors > 0) ? "error" : (NumWarnings > 0) ? "warning" : null;
+				if (NewOutcome != Outcome)
+				{
+					SetProperty("/myJobStep/outcome", NewOutcome);
+					Outcome = NewOutcome;
+				}
+
+				// Wait until the next update
+				UpdateEvent.WaitOne();
+			}
+		}
+
+		/// <summary>
+		/// Sets an EC property 
+		/// </summary>
+		/// <param name="Command">The command to run</param>
+		void SetProperty(string Name, string Value)
+		{
+			if (bUpdateProperties)
+			{
+				string Arguments = String.Format("setProperty \"{0}\" \"{1}\"", Name, Value);
+				using (Process Process = Process.Start("ectool", Arguments))
+				{
+					Process.WaitForExit();
+
+					if (Process.ExitCode != 0)
+					{
+						Log.TraceWarning("ectool {0} terminated with exit code {1}", Arguments, Process.ExitCode);
+					}
+				}
 			}
 		}
 
 		/// <summary>
 		/// Write the list of errors to disk in XML format
 		/// </summary>
-		public void Write()
+		static void Write(FileReference OutputFile, List<ErrorMatch> Errors)
 		{
 			XmlWriterSettings Settings = new XmlWriterSettings();
 			Settings.NewLineChars = Environment.NewLine;
diff --git a/Engine/Source/Programs/BuildAgent/Run/RunMode.cs b/Engine/Source/Programs/BuildAgent/Run/RunMode.cs
--- a/Engine/Source/Programs/BuildAgent/Run/RunMode.cs
+++ b/Engine/Source/Programs/BuildAgent/Run/RunMode.cs
@@ -35,14 +35,16 @@ namespace BuildAgent.Run
 		[Description("The program to run.")]
 		FileReference Program = null;
 
-		[CommandLine("-Arguments=")]
-		[Description("Specifies command line arguments to pass to the child process.")]
-		string Arguments = "";
+		string[] ProgramArguments;
 
 		[CommandLine]
 		[Description("Amount of time to leave before killing the child process.")]
 		TimeSpan? Timeout = null;
 
+		[CommandLine("-NoWarnings")]
+		[Description("Ignores any warnings")]
+		bool bNoWarnings = false;
+
 		[CommandLine("-DebugListener")]
 		[Description("Enables the debug listener")]
 		bool bDebugListener = false;
@@ -60,10 +62,12 @@ namespace BuildAgent.Run
 				throw new CommandLineArgumentException(String.Format("Either -{0}=... or -{1}=... must be specified.", nameof(InputFile), nameof(Program)));
 			}
 
-			if (!bElectricCommanderListener && !String.IsNullOrEmpty(Environment.GetEnvironmentVariable("COMMANDER_JOBSTEP")))
+			if (!bElectricCommanderListener && !String.IsNullOrEmpty(Environment.GetEnvironmentVariable("COMMANDER_JOBSTEPID")))
 			{
 				bElectricCommanderListener = true;
 			}
+
+			ProgramArguments = Arguments.GetPositionalArguments();
 		}
 
 		public override void Execute()
@@ -131,7 +135,7 @@ namespace BuildAgent.Run
 					using (StreamReader Reader = new StreamReader(InputFile.FullName))
 					{
 						LineFilter Filter = new LineFilter(() => Reader.ReadLine());
-						ProcessErrors(Filter.ReadLine, Matchers, IgnorePatterns, Listeners);
+						ProcessErrors(Filter.ReadLine, Matchers, IgnorePatterns, Listeners, bNoWarnings);
 					}
 				}
 				else
@@ -143,10 +147,10 @@ namespace BuildAgent.Run
 					}
 
 					CancellationToken CancellationToken = CancellationTokenSource.Token;
-					using (ManagedProcess Process = new ManagedProcess(null, Program.FullName, Arguments, null, null, null, ProcessPriorityClass.Normal))
+					using (ManagedProcess Process = new ManagedProcess(null, Program.FullName, CommandLineArguments.Join(ProgramArguments), null, null, null, ProcessPriorityClass.Normal))
 					{
 						Func<string> ReadLine = new LineFilter(() => ReadProcessLine(Process, CancellationToken)).ReadLine;
-						ProcessErrors(ReadLine, Matchers, IgnorePatterns, Listeners);
+						ProcessErrors(ReadLine, Matchers, IgnorePatterns, Listeners, bNoWarnings);
 					}
 
 					ProcessUtils.TerminateChildProcesses();
@@ -184,7 +188,8 @@ namespace BuildAgent.Run
 		/// <param name="Matchers">List of matchers to run against the text</param>
 		/// <param name="IgnorePatterns">List of patterns to ignore</param>
 		/// <param name="Listeners">Set of listeners for processing the errors</param>
-		static void ProcessErrors(Func<string> ReadLine, List<IErrorMatcher> Matchers, List<string> IgnorePatterns, List<IErrorListener> Listeners)
+		/// <param name="bNoWarnings">Does not output warnings</param>
+		static void ProcessErrors(Func<string> ReadLine, List<IErrorMatcher> Matchers, List<string> IgnorePatterns, List<IErrorListener> Listeners, bool bNoWarnings)
 		{
 			System.Text.RegularExpressions.Regex.CacheSize = 1000;
 
@@ -210,6 +215,12 @@ namespace BuildAgent.Run
 					}
 				}
 
+				// If we matched a warning and don't want it, clear it out
+				if (Error != null && Error.Severity == ErrorSeverity.Warning && bNoWarnings)
+				{
+					Error = null;
+				}
+
 				// If we did match something, check if it's not negated by an ignore pattern. We typically have relatively few errors and many more ignore patterns than matchers, so it's quicker 
 				// to check them in response to an identified error than to treat them as matchers of their own.
 				if (Error != null)
diff --git a/Engine/Source/Programs/DotNETCommon/DotNETUtilities/CommandLineArguments.cs b/Engine/Source/Programs/DotNETCommon/DotNETUtilities/CommandLineArguments.cs
--- a/Engine/Source/Programs/DotNETCommon/DotNETUtilities/CommandLineArguments.cs
+++ b/Engine/Source/Programs/DotNETCommon/DotNETUtilities/CommandLineArguments.cs
@@ -102,6 +102,11 @@ namespace Tools.DotNETCommon
 		/// </summary>
 		int[] NextArgumentIndex;
 
+		/// <summary>
+		/// List of positional arguments
+		/// </summary>
+		List<int> PositionalArgumentIndices = new List<int>();
+
 		/// <summary>
 		/// Array of characters that separate argument names from values
 		/// </summary>
@@ -131,40 +136,59 @@ namespace Tools.DotNETCommon
 			ArgumentToFirstIndex = new Dictionary<string, int>(Arguments.Length, StringComparer.OrdinalIgnoreCase);
 			for(int Idx = 0; Idx < Arguments.Length; Idx++)
 			{
-				int SeparatorIdx = Arguments[Idx].IndexOfAny(ValueSeparators);
-				if(SeparatorIdx == -1)
+				if (Arguments[Idx].Equals("--", StringComparison.Ordinal))
 				{
-					// Ignore duplicate -Option flags; they are harmless.
-					if(ArgumentToFirstIndex.ContainsKey(Arguments[Idx]))
-					{
-						UsedArguments.Set(Idx, true);
-					}
-					else
+					// End of option arguments
+					MarkAsUsed(Idx++);
+					for (; Idx < Arguments.Length; Idx++)
 					{
-						ArgumentToFirstIndex.Add(Arguments[Idx], Idx);
+						PositionalArgumentIndices.Add(Idx);
 					}
-
-					// Mark this argument as a flag
-					FlagArguments.Set(Idx, true);
+					break;
 				}
-				else
+				else if (Arguments[Idx].StartsWith("-", StringComparison.Ordinal))
 				{
-					// Just take the part up to and including the separator character
-					string Prefix = Arguments[Idx].Substring(0, SeparatorIdx + 1);
-
-					// Add the prefix to the argument lookup, or update the appropriate matching argument list if it's been seen before
-					int ExistingArgumentIndex;
-					if(ArgumentToFirstIndex.TryGetValue(Prefix, out ExistingArgumentIndex))
+					// Option argument
+					int SeparatorIdx = Arguments[Idx].IndexOfAny(ValueSeparators);
+					if (SeparatorIdx == -1)
 					{
-						NextArgumentIndex[LastArgumentIndex[ExistingArgumentIndex]] = Idx;
-						LastArgumentIndex[ExistingArgumentIndex] = Idx;
+						// Ignore duplicate -Option flags; they are harmless.
+						if (ArgumentToFirstIndex.ContainsKey(Arguments[Idx]))
+						{
+							UsedArguments.Set(Idx, true);
+						}
+						else
+						{
+							ArgumentToFirstIndex.Add(Arguments[Idx], Idx);
+						}
+
+						// Mark this argument as a flag
+						FlagArguments.Set(Idx, true);
 					}
 					else
 					{
-						ArgumentToFirstIndex.Add(Prefix, Idx);
-						LastArgumentIndex[Idx] = Idx;
+						// Just take the part up to and including the separator character
+						string Prefix = Arguments[Idx].Substring(0, SeparatorIdx + 1);
+
+						// Add the prefix to the argument lookup, or update the appropriate matching argument list if it's been seen before
+						int ExistingArgumentIndex;
+						if (ArgumentToFirstIndex.TryGetValue(Prefix, out ExistingArgumentIndex))
+						{
+							NextArgumentIndex[LastArgumentIndex[ExistingArgumentIndex]] = Idx;
+							LastArgumentIndex[ExistingArgumentIndex] = Idx;
+						}
+						else
+						{
+							ArgumentToFirstIndex.Add(Prefix, Idx);
+							LastArgumentIndex[Idx] = Idx;
+						}
 					}
 				}
+				else
+				{
+					// Positional argument
+					PositionalArgumentIndices.Add(Idx);
+				}
 			}
 		}
 
@@ -185,7 +209,7 @@ namespace Tools.DotNETCommon
 		{
 			get { return Arguments[Index]; }
 		}
-		
+
 		/// <summary>
 		/// Determines if an argument has been used
 		/// </summary>
@@ -241,6 +265,41 @@ namespace Tools.DotNETCommon
 			return ArgumentToFirstIndex.ContainsKey(Prefix);
 		}
 
+		/// <summary>
+		/// Gets the positional argument at the given index
+		/// </summary>
+		/// <returns>Number of positional arguments</returns>
+		public int GetPositionalArgumentCount()
+		{
+			return PositionalArgumentIndices.Count;
+		}
+
+		/// <summary>
+		/// Gets the index of the numbered positional argument
+		/// </summary>
+		/// <param name="Index">Number of the positional argument</param>
+		/// <returns>Index of the positional argument</returns>
+		public int GetPositionalArgumentIndex(int Num)
+		{
+			return PositionalArgumentIndices[Num];
+		}
+
+		/// <summary>
+		/// Returns all the positional arguments, and marks them as used
+		/// </summary>
+		/// <returns>Array of positional arguments</returns>
+		public string[] GetPositionalArguments()
+		{
+			string[] PositionalArguments = new string[PositionalArgumentIndices.Count];
+			for (int Idx = 0; Idx < PositionalArguments.Length; Idx++)
+			{
+				int Index = PositionalArgumentIndices[Idx];
+				MarkAsUsed(Index);
+				PositionalArguments[Idx] = Arguments[Index];
+			}
+			return PositionalArguments;
+		}
+
 		/// <summary>
 		/// Gets the value specified by an argument with the given prefix. Throws an exception if the argument was not specified.
 		/// </summary>
@@ -764,14 +823,38 @@ namespace Tools.DotNETCommon
 		/// <returns>Argument which is safe to pass on the command line</returns>
 		public static string Quote(string Argument)
 		{
-			if (Argument.IndexOf(' ') != -1 && Argument[0] != '\"')
+			// See if the entire string is quoted correctly
+			bool bInQuotes = false;
+			for (int Idx = 0;;Idx++)
 			{
-				return "\"" + Argument+ "\"";
+				if (Idx == Argument.Length)
+				{
+					return Argument;
+				}
+				else if (Argument[Idx] == '\"')
+				{
+					bInQuotes ^= true;
+				}
+				else if (Argument[Idx] == ' ')
+				{
+					break;
+				}
 			}
-			else
+
+			// Try to insert a quote after the argument string
+			if (Argument[0] == '-')
 			{
-				return Argument;
+				for(int Idx = 1; Idx < Argument.Length && Argument[Idx] != ' '; Idx++)
+				{
+					if (Argument[Idx] == '=')
+					{
+						return String.Format("{0}=\"{1}\"", Argument.Substring(0, Idx), Argument.Substring(Idx + 1).Replace("\"", "\\\""));
+					}
+				}
 			}
+
+			// Quote the whole thing
+			return "\"" + Argument.Replace("\"", "\\\"") + "\"";
 		}
 
 		/// <summary>
