Add functionality to issue commands to UGS servers to BuildAgent tool.

diff --git a/Engine/Source/Programs/BuildAgent/BuildAgent.csproj b/Engine/Source/Programs/BuildAgent/BuildAgent.csproj
--- a/Engine/Source/Programs/BuildAgent/BuildAgent.csproj
+++ b/Engine/Source/Programs/BuildAgent/BuildAgent.csproj
@@ -36,6 +36,9 @@
   <ItemGroup>
     <Reference Include="System" />
     <Reference Include="System.Core" />
+    <Reference Include="System.Runtime.Serialization" />
+    <Reference Include="System.Web" />
+    <Reference Include="System.Web.Extensions" />
     <Reference Include="System.Xml.Linq" />
     <Reference Include="System.Data.DataSetExtensions" />
     <Reference Include="Microsoft.CSharp" />
@@ -47,6 +50,8 @@
     <Compile Include="..\..\..\Source\Programs\DotNETCommon\MetaData.cs">
       <Link>Properties\MetaData.cs</Link>
     </Compile>
+    <Compile Include="Server\ApiTypes.cs" />
+    <Compile Include="Server\ServerMode.cs" />
     <Compile Include="Run\RunMode.cs" />
     <Compile Include="ProgramMode.cs" />
     <Compile Include="Run\AutoRegisterAttribute.cs" />
@@ -54,7 +59,7 @@
     <Compile Include="Run\LineFilter.cs" />
     <Compile Include="Run\Listeners\ElectricCommanderListener.cs" />
     <Compile Include="Run\Listeners\DebugOutputListener.cs" />
-	<Compile Include="Run\Matchers\**\*.cs" />
+    <Compile Include="Run\Matchers\**\*.cs" />
     <Compile Include="Program.cs" />
     <Compile Include="Run\ErrorMatch.cs" />
     <Compile Include="Run\Interfaces\IErrorListener.cs" />
diff --git a/Engine/Source/Programs/BuildAgent/Program.cs b/Engine/Source/Programs/BuildAgent/Program.cs
--- a/Engine/Source/Programs/BuildAgent/Program.cs
+++ b/Engine/Source/Programs/BuildAgent/Program.cs
@@ -112,7 +112,7 @@ namespace BuildAgent
 			// If the help flag is specified, print the help info and exit immediately
 			if (ModeArguments.HasOption("-Help"))
 			{
-				HelpUtils.PrintHelp(CommandName, HelpUtils.GetDescription(CommandType), Mode.GetParameters());
+				HelpUtils.PrintHelp(CommandName, HelpUtils.GetDescription(CommandType), Mode.GetParameters(ModeArguments));
 				return 1;
 			}
 
@@ -128,7 +128,7 @@ namespace BuildAgent
 				Log.TraceInformation("");
 				Log.TraceInformation("Arguments for {0}:", CommandName);
 
-				HelpUtils.PrintTable(Mode.GetParameters(), 4, 24);
+				HelpUtils.PrintTable(Mode.GetParameters(ModeArguments), 4, 24);
 				return 1;
 			}
 
diff --git a/Engine/Source/Programs/BuildAgent/ProgramMode.cs b/Engine/Source/Programs/BuildAgent/ProgramMode.cs
--- a/Engine/Source/Programs/BuildAgent/ProgramMode.cs
+++ b/Engine/Source/Programs/BuildAgent/ProgramMode.cs
@@ -17,7 +17,7 @@ namespace BuildAgent
 			Arguments.ApplyTo(this);
 		}
 
-		public virtual List<KeyValuePair<string, string>> GetParameters()
+		public virtual List<KeyValuePair<string, string>> GetParameters(CommandLineArguments Arguments)
 		{
 			return CommandLineArguments.GetParameters(GetType());
 		}
diff --git a/Engine/Source/Programs/BuildAgent/Server/ApiTypes.cs b/Engine/Source/Programs/BuildAgent/Server/ApiTypes.cs
new file mode 100644
--- /dev/null
+++ b/Engine/Source/Programs/BuildAgent/Server/ApiTypes.cs
@@ -0,0 +1,93 @@
+// Copyright 1998-2019 Epic Games, Inc. All Rights Reserved.
+
+using System;
+using System.Collections.Generic;
+using System.Linq;
+using System.Text;
+using System.Threading.Tasks;
+
+namespace BuildAgent.WebApi
+{
+	[AttributeUsage(AttributeTargets.Field)]
+	class OptionalAttribute : Attribute
+	{
+	}
+
+#pragma warning disable CS0649
+
+	static class ApiTypes
+	{
+		public const int MaxIssueDetailsLength = 8192;
+
+		public enum Outcome
+		{
+			Success = 1,
+			Error = 2,
+			Warning = 3,
+		}
+
+		public class AddIssue
+		{
+			public string Project;
+			public string Summary;
+			[Optional] public string Owner;
+		}
+
+		public class AddIssueResponse
+		{
+			public long Id;
+		}
+
+		public class UpdateIssue
+		{
+			[Optional] public bool? Acknowledged;
+			[Optional] public int? FixChange;
+			[Optional] public string Owner;
+			[Optional] public string NominatedBy;
+			[Optional] public bool? Resolved;
+			[Optional] public string Summary;
+			[Optional] public string Url;
+		}
+
+		public class AddBuild
+		{
+			public string Stream;
+			public int Change;
+			public string JobName;
+			public string JobUrl;
+			public string JobStepName;
+			public string JobStepUrl;
+			[Optional] public string ErrorUrl;
+			public Outcome Outcome;
+		}
+
+		public class AddBuildResponse
+		{
+			public long Id;
+		}
+
+		public class UpdateBuild
+		{
+			[Optional] public int? Outcome;
+		}
+
+		public class AddDiagnostic
+		{
+			[Optional] public long? BuildId;
+			public string Message;
+			public string Url;
+		}
+
+		public class AddDiagnosticResponse
+		{
+			public long Id;
+		}
+
+		public class Watcher
+		{
+			public string UserName;
+		}
+	}
+
+#pragma warning restore CS0649
+}
diff --git a/Engine/Source/Programs/BuildAgent/Server/ServerMode.cs b/Engine/Source/Programs/BuildAgent/Server/ServerMode.cs
new file mode 100644
--- /dev/null
+++ b/Engine/Source/Programs/BuildAgent/Server/ServerMode.cs
@@ -0,0 +1,333 @@
+using System;
+using System.Collections.Generic;
+using System.ComponentModel;
+using System.IO;
+using System.Linq;
+using System.Net;
+using System.Reflection;
+using System.Text;
+using System.Threading.Tasks;
+using System.Web.Script.Serialization;
+using Tools.DotNETCommon;
+
+namespace BuildAgent.WebApi
+{
+	[ProgramMode("Server", "Issues direct API commands to the UGS web service")]
+	class ServerMode : ProgramMode
+	{
+		class CommandInfo
+		{
+			public readonly string Name;
+			public readonly string Verb;
+			public readonly string Resource;
+			public readonly Type BodyType;
+			public readonly string[] OptionalParams;
+
+			public CommandInfo(string Name, string Verb, string Resource, Type BodyType = null, string[] OptionalParams = null)
+			{
+				this.Name = Name;
+				this.Verb = Verb;
+				this.Resource = Resource;
+				this.BodyType = BodyType;
+				this.OptionalParams = OptionalParams;
+			}
+		}
+
+		static readonly List<CommandInfo> Commands = new List<CommandInfo>
+		{
+			// Issues
+			new CommandInfo("AddIssue", "POST", "/api/issues", typeof(ApiTypes.AddIssue)),
+			new CommandInfo("GetIssue", "GET", "/api/issues/{id}"),
+			new CommandInfo("GetIssues", "GET", "/api/issues", OptionalParams: new string[] { "includeresolved", "maxresults", "user" }),
+			new CommandInfo("UpdateIssue", "PUT", "/api/issues/{Id}", typeof(ApiTypes.UpdateIssue)),
+			new CommandInfo("DeleteIssue", "DELETE", "/api/issues/{Id}"),
+
+			// Builds
+			new CommandInfo("AddBuild", "POST", "/api/issues/{Issue}/builds", typeof(ApiTypes.AddBuild)),
+			new CommandInfo("GetBuilds", "GET", "/api/issues/{Issue}/builds"),
+
+			// Individual builds
+			new CommandInfo("GetBuild", "GET", "/api/issuebuilds/{Build}"),
+			new CommandInfo("UpdateBuild", "PUT", "/api/issuebuilds/{Build}", typeof(ApiTypes.UpdateBuild)),
+
+			// Diagnostics
+			new CommandInfo("AddDiagnostic", "POST", "/api/issues/{Issue}/diagnostics", typeof(ApiTypes.AddDiagnostic)),
+			new CommandInfo("GetDiagnostics", "GET", "/api/issues/{Issue}/diagnostics"),
+
+			// Watchers
+			new CommandInfo("GetWatchers", "GET", "/api/issues/{Issue}/watchers"),
+			new CommandInfo("AddWatcher", "POST", "/api/issues/{Issue}/watchers", typeof(ApiTypes.Watcher)),
+			new CommandInfo("RemoveWatcher", "DELETE", "/api/issues/{Issue}/watchers", typeof(ApiTypes.Watcher)),
+		};
+
+		[CommandLine("-Server=")]
+		[Description("Url of the server to connect to")]
+		public string Server = null;
+
+		CommandInfo Command;
+		StringBuilder CommandUrl;
+		string BodyText;
+
+		public override void Configure(CommandLineArguments Arguments)
+		{
+			base.Configure(Arguments);
+
+			// Get the command line
+			if (!TryGetCommandArgument(Arguments, out Command))
+			{
+				StringBuilder Message = new StringBuilder("Missing or invalid command name. Valid commands are:");
+				foreach (CommandInfo AvailableCommand in Commands)
+				{
+					Message.AppendFormat("\n  {0}", AvailableCommand.Name);
+				}
+				throw new CommandLineArgumentException(Message.ToString());
+			}
+
+			// Parse the server argument
+			if (String.IsNullOrEmpty(Server))
+			{
+				Server = Environment.GetEnvironmentVariable("UGS_METADATA_SERVER_URL");
+				if (String.IsNullOrEmpty(Server))
+				{
+					throw new CommandLineArgumentException("Missing -Server=... argument.");
+				}
+			}
+
+			// Create the command url
+			CommandUrl = new StringBuilder();
+			CommandUrl.AppendFormat("{0}{1}", Server, Command.Resource);
+
+			// Replace all the arguments that are embedded into the resource name
+			foreach (string ResourceArg in ParseArgumentNamesFromResource(Command.Resource))
+			{
+				string Value = Arguments.GetString(String.Format("-{0}=", ResourceArg));
+				CommandUrl.Replace("{" + ResourceArg + "}", Value);
+			}
+
+			// Add all the required and optional parameters
+			List<string> QueryParams = new List<string>();
+			if (Command.OptionalParams != null)
+			{
+				foreach (string OptionalParam in Command.OptionalParams)
+				{
+					string Value = Arguments.GetStringOrDefault(String.Format("-{0}=", OptionalParam), null);
+					if (Value != null)
+					{
+						QueryParams.Add(String.Format("{0}={1}", OptionalParam, Value));
+					}
+				}
+			}
+
+			// Append the parameters to the URL
+			for (int Idx = 0; Idx < QueryParams.Count; Idx++)
+			{
+				if (Idx == 0)
+				{
+					CommandUrl.Append("?");
+				}
+				else
+				{
+					CommandUrl.Append("&");
+				}
+				CommandUrl.Append(QueryParams[Idx]);
+			}
+
+			// Parse additional options for the message body
+			if (Command.BodyType != null)
+			{
+				object BodyObject = ParseObject(Arguments, Command.BodyType);
+				BodyText = new JavaScriptSerializer().Serialize(BodyObject);
+			}
+		}
+
+		public override List<KeyValuePair<string, string>> GetParameters(CommandLineArguments Arguments)
+		{
+			List<KeyValuePair<string, string>> Parameters = new List<KeyValuePair<string, string>>();
+
+			if(TryGetCommandArgument(Arguments, out Command))
+			{
+				Parameters.Insert(0, new KeyValuePair<string, string>(Command.Name, "Name of the command to execute"));
+
+				foreach (string ArgumentName in ParseArgumentNamesFromResource(Command.Resource))
+				{
+					Parameters.Add(new KeyValuePair<string, string>(String.Format("-{0}", ArgumentName), "Required"));
+				}
+				if (Command.OptionalParams != null)
+				{
+					foreach (string ArgumentName in Command.OptionalParams)
+					{
+						Parameters.Add(new KeyValuePair<string, string>(String.Format("-{0}", ArgumentName), "Optional"));
+					}
+				}
+				if (Command.BodyType != null)
+				{
+					foreach (FieldInfo Field in Command.BodyType.GetFields(BindingFlags.Instance | BindingFlags.Public))
+					{
+						if (IsOptionalObjectField(Field))
+						{
+							Parameters.Add(new KeyValuePair<string, string>(String.Format("-{0}", Field.Name), String.Format("{0} (optional)", GetFieldValueDescription(Field.FieldType))));
+						}
+						else
+						{
+							Parameters.Add(new KeyValuePair<string, string>(String.Format("-{0}", Field.Name), GetFieldValueDescription(Field.FieldType)));
+						}
+					}
+				}
+			}
+			else
+			{
+				Parameters.Insert(0, new KeyValuePair<string, string>("Command", "Name of the command to execute"));
+			}
+
+			return Parameters;
+		}
+
+		public override void Execute()
+		{
+			// Create the request
+			HttpWebRequest Request = (HttpWebRequest)WebRequest.Create(CommandUrl.ToString());
+			Request.ContentType = "application/json";
+			Request.Method = Command.Verb;
+			if (BodyText != null)
+			{
+				byte[] BodyData = Encoding.UTF8.GetBytes(BodyText);
+				using (Stream RequestStream = Request.GetRequestStream())
+				{
+					RequestStream.Write(BodyData, 0, BodyData.Length);
+				}
+			}
+
+			// Read the response
+			HttpWebResponse Response = (HttpWebResponse)Request.GetResponse();
+			Console.WriteLine("Response: {0} ({1})", (int)Response.StatusCode, Response.StatusDescription);
+			using (StreamReader ResponseReader = new StreamReader(Response.GetResponseStream(), Encoding.Default))
+			{
+				string ResponseContent = ResponseReader.ReadToEnd();
+				if (!String.IsNullOrEmpty(ResponseContent))
+				{
+					Console.WriteLine(Json.Format(ResponseContent));
+				}
+			}
+		}
+
+		static bool TryGetCommandArgument(CommandLineArguments Arguments, out CommandInfo Command)
+		{
+			string CommandName;
+			if (!Arguments.TryGetPositionalArgument(out CommandName))
+			{
+				Command = null;
+				return false;
+			}
+
+			CommandInfo FoundCommand = Commands.FirstOrDefault(x => x.Name.Equals(CommandName, StringComparison.OrdinalIgnoreCase));
+			if (FoundCommand == null)
+			{
+				Command = null;
+				return false;
+			}
+
+			Command = FoundCommand;
+			return true;
+		}
+
+		static IEnumerable<string> ParseArgumentNamesFromResource(string Resource)
+		{
+			for (int ArgIdx = Resource.IndexOf('{'); ArgIdx != -1; ArgIdx = Resource.IndexOf('{', ArgIdx))
+			{
+				int ArgEndIdx = Resource.IndexOf('}', ArgIdx + 1);
+				yield return Resource.Substring(ArgIdx + 1, ArgEndIdx - (ArgIdx + 1));
+				ArgIdx = ArgEndIdx + 1;
+			}
+		}
+
+		static object ParseObject(CommandLineArguments Arguments, Type ObjectType)
+		{
+			object Instance = Activator.CreateInstance(ObjectType);
+			foreach (FieldInfo Field in Instance.GetType().GetFields(BindingFlags.Instance | BindingFlags.Public))
+			{
+				string Value = Arguments.GetStringOrDefault(String.Format("-{0}=", Field.Name), null);
+				if (Value == null)
+				{
+					if (!IsOptionalObjectField(Field))
+					{
+						throw new CommandLineArgumentException(String.Format("Missing -{0}=... argument", Field.Name));
+					}
+				}
+				else
+				{
+					Type FieldType = Field.FieldType;
+					if (FieldType.IsGenericType && FieldType.GetGenericTypeDefinition() == typeof(Nullable<>))
+					{
+						FieldType = FieldType.GetGenericArguments()[0];
+					}
+
+					if (FieldType == typeof(string))
+					{
+						Field.SetValue(Instance, Value);
+					}
+					else if (FieldType == typeof(int))
+					{
+						int IntValue;
+						if (!int.TryParse(Value, out IntValue))
+						{
+							throw new CommandLineArgumentException(String.Format("Invalid value for '{0}' - {1} is not an integer.", Field.Name, Value));
+						}
+						Field.SetValue(Instance, IntValue);
+					}
+					else if (FieldType == typeof(bool))
+					{
+						bool BoolValue;
+						if (!bool.TryParse(Value, out BoolValue))
+						{
+							throw new CommandLineArgumentException(String.Format("Invalid value for '{0}' - {1} is not a boolean.", Field.Name, Value));
+						}
+						Field.SetValue(Instance, BoolValue);
+					}
+					else if (FieldType.IsEnum)
+					{
+						object EnumValue;
+						try
+						{
+							EnumValue = Enum.Parse(FieldType, Value, true);
+						}
+						catch (ArgumentException)
+						{
+							throw new CommandLineArgumentException(String.Format("Invalid value for '{0}' - should be {1}", Field.Name, String.Join("/", Enum.GetNames(FieldType))));
+						}
+						Field.SetValue(Instance, EnumValue);
+					}
+					else
+					{
+						throw new NotImplementedException(String.Format("Unsupported type '{0}'", FieldType));
+					}
+				}
+			}
+			return Instance;
+		}
+
+		static bool IsOptionalObjectField(FieldInfo Field)
+		{
+			if (Field.GetCustomAttribute<OptionalAttribute>() != null)
+			{
+				return true;
+			}
+			if (Field.FieldType.IsGenericType && Field.FieldType.GetGenericTypeDefinition() == typeof(Nullable<>))
+			{
+				return true;
+			}
+			return false;
+		}
+
+		static string GetFieldValueDescription(Type FieldType)
+		{
+			if (FieldType.IsEnum)
+			{
+				return String.Join("|", Enum.GetNames(FieldType));
+			}
+			else
+			{
+				return "...";
+			}
+		}
+	}
+}
diff --git a/Engine/Source/Programs/DotNETCommon/DotNETUtilities/CommandLineArguments.cs b/Engine/Source/Programs/DotNETCommon/DotNETUtilities/CommandLineArguments.cs
--- a/Engine/Source/Programs/DotNETCommon/DotNETUtilities/CommandLineArguments.cs
+++ b/Engine/Source/Programs/DotNETCommon/DotNETUtilities/CommandLineArguments.cs
@@ -284,6 +284,28 @@ namespace Tools.DotNETCommon
 			return PositionalArgumentIndices[Num];
 		}
 
+		/// <summary>
+		/// Attempts to read the next unused positional argument
+		/// </summary>
+		/// <param name="Argument">Receives the argument that was read, on success</param>
+		/// <returns>True if an argument was read</returns>
+		public bool TryGetPositionalArgument(out string Argument)
+		{
+			for (int Idx = 0; Idx < PositionalArgumentIndices.Count; Idx++)
+			{
+				int Index = PositionalArgumentIndices[Idx];
+				if (!HasBeenUsed(Index))
+				{
+					MarkAsUsed(Index);
+					Argument = Arguments[Index];
+					return true;
+				}
+			}
+
+			Argument = null;
+			return false;
+		}
+
 		/// <summary>
 		/// Returns all the positional arguments, and marks them as used
 		/// </summary>
