Prevent EC listener from flushing constantly.

#rnx

diff --git a/Engine/Source/Programs/BuildAgent/Run/Listeners/ElectricCommanderListener.cs b/Engine/Source/Programs/BuildAgent/Run/Listeners/ElectricCommanderListener.cs
--- a/Engine/Source/Programs/BuildAgent/Run/Listeners/ElectricCommanderListener.cs
+++ b/Engine/Source/Programs/BuildAgent/Run/Listeners/ElectricCommanderListener.cs
@@ -48,7 +48,7 @@ namespace BuildAgent.Listeners
 		/// <summary>
 		/// Event that is set whenever EC needs to be updated
 		/// </summary>
-		ManualResetEvent UpdateEvent = new ManualResetEvent(false);
+		AutoResetEvent UpdateEvent = new AutoResetEvent(false);
 
 		/// <summary>
 		/// Set when the object is being disposed
@@ -95,6 +95,12 @@ namespace BuildAgent.Listeners
 
 				BackgroundThread = null;
 			}
+
+			if(UpdateEvent != null)
+			{
+				UpdateEvent.Dispose();
+				UpdateEvent = null;
+			}
 		}
 
 		/// <summary>
