Add functionality for updating build issues into the BuildAgent tool.

diff --git a/Engine/Source/Programs/BuildAgent/BuildAgent.csproj b/Engine/Source/Programs/BuildAgent/BuildAgent.csproj
--- a/Engine/Source/Programs/BuildAgent/BuildAgent.csproj
+++ b/Engine/Source/Programs/BuildAgent/BuildAgent.csproj
@@ -22,6 +22,7 @@
     <DefineConstants>DEBUG;TRACE</DefineConstants>
     <ErrorReport>prompt</ErrorReport>
     <WarningLevel>4</WarningLevel>
+    <NoWarn>CS0649</NoWarn>
   </PropertyGroup>
   <PropertyGroup Condition=" '$(Configuration)|$(Platform)' == 'Development|AnyCPU' ">
     <PlatformTarget>AnyCPU</PlatformTarget>
@@ -50,6 +51,22 @@
     <Compile Include="..\..\..\Source\Programs\DotNETCommon\MetaData.cs">
       <Link>Properties\MetaData.cs</Link>
     </Compile>
+    <Compile Include="Issues\Issue.cs" />
+    <Compile Include="Issues\IssueDiagnostic.cs" />
+    <Compile Include="Issues\IssueBuild.cs" />
+    <Compile Include="Issues\IssueHistory.cs" />
+    <Compile Include="Issues\PersistentState.cs" />
+    <Compile Include="Issues\ChangeInfo.cs" />
+    <Compile Include="Issues\InputData.cs" />
+    <Compile Include="Issues\Matchers\CompileIssueMatcher.cs" />
+    <Compile Include="Issues\Matchers\ContentIssueMatcher.cs" />
+    <Compile Include="Issues\Matchers\CopyrightNoticeMatcher.cs" />
+    <Compile Include="Issues\Matchers\DefaultIssueMatcher.cs" />
+    <Compile Include="Issues\Matchers\GenericCodeIssueMatcher.cs" />
+    <Compile Include="Issues\Matchers\GenericContentIssueMatcher.cs" />
+    <Compile Include="Issues\Matchers\UndefinedSymbolMatcher.cs" />
+    <Compile Include="Issues\Matcher.cs" />
+    <Compile Include="Issues\UpdateIssuesMode.cs" />
     <Compile Include="Server\ApiTypes.cs" />
     <Compile Include="Server\ServerMode.cs" />
     <Compile Include="Run\RunMode.cs" />
diff --git a/Engine/Source/Programs/BuildAgent/Issues/ChangeInfo.cs b/Engine/Source/Programs/BuildAgent/Issues/ChangeInfo.cs
new file mode 100644
--- /dev/null
+++ b/Engine/Source/Programs/BuildAgent/Issues/ChangeInfo.cs
@@ -0,0 +1,21 @@
+// Copyright 1998-2019 Epic Games, Inc. All Rights Reserved.
+
+using System;
+using System.Collections.Generic;
+using System.Linq;
+using System.Text;
+using System.Threading.Tasks;
+using Tools.DotNETCommon.Perforce;
+
+namespace BuildAgent.Issues
+{
+	/// <summary>
+	/// Stores information about a change and its merge history
+	/// </summary>
+	class ChangeInfo
+	{
+		public ChangesRecord Record;
+		public List<int> SourceChanges = new List<int>();
+		public DescribeRecord CachedDescribeRecord;
+	}
+}
diff --git a/Engine/Source/Programs/BuildAgent/Issues/InputData.cs b/Engine/Source/Programs/BuildAgent/Issues/InputData.cs
new file mode 100644
--- /dev/null
+++ b/Engine/Source/Programs/BuildAgent/Issues/InputData.cs
@@ -0,0 +1,118 @@
+// Copyright 1998-2019 Epic Games, Inc. All Rights Reserved.
+
+using System;
+using System.Collections.Generic;
+using System.Diagnostics;
+using System.Linq;
+using System.Runtime.Serialization;
+using System.Text;
+using System.Threading.Tasks;
+
+namespace BuildAgent.Issues
+{
+	/// <summary>
+	/// A diagnostic message produced during the build
+	/// </summary>
+	[DataContract]
+	[DebuggerDisplay("{Type}: {Message}")]
+	public class InputDiagnostic
+	{
+		[DataMember(IsRequired = true)]
+		public string Type;
+
+		[DataMember(IsRequired = true)]
+		public string Message;
+
+		[DataMember(IsRequired = true)]
+		public string Url;
+	}
+
+	/// <summary>
+	/// Represents the outcome of a build step
+	/// </summary>
+	[DataContract]
+	[DebuggerDisplay("{Name}")]
+	public class InputJobStep
+	{
+		/// <summary>
+		/// Name of the build step. Multiple step results can be submitted for a single build, and can be used to augment a single issue with additional diagnostic information.
+		/// </summary>
+		[DataMember(IsRequired = true)]
+		public string Name;
+
+		/// <summary>
+		/// Url of this job step
+		/// </summary>
+		[DataMember(IsRequired = true)]
+		public string Url;
+
+		/// <summary>
+		/// Base directory that the build was executed in. Optional. This is used to determine branch-relative file paths for diagnostics.
+		/// </summary>
+		[DataMember]
+		public string BaseDirectory;
+
+		/// <summary>
+		/// Output from the build, typically messages which are already designated warnings or errors.
+		/// </summary>
+		[DataMember]
+		public List<InputDiagnostic> Diagnostics = new List<InputDiagnostic>();
+	}
+
+	/// <summary>
+	/// Class containing input data for a job
+	/// </summary>
+	[DataContract]
+	[DebuggerDisplay("{Stream}: {Change} - {Name}")]
+	class InputJob
+	{
+		/// <summary>
+		/// Project that this build belongs to
+		/// </summary>
+		[DataMember(IsRequired = true)]
+		public string Project;
+
+		/// <summary>
+		/// Name of this build. Displayed in the UI.
+		/// </summary>
+		[DataMember(IsRequired = true)]
+		public string Name;
+
+		/// <summary>
+		/// Url for this job
+		/// </summary>
+		[DataMember(IsRequired = true)]
+		public string Url;
+
+		/// <summary>
+		/// Stream that was built
+		/// </summary>
+		[DataMember(IsRequired = true)]
+		public string Stream;
+
+		/// <summary>
+		/// Changelist number that was built
+		/// </summary>
+		[DataMember(IsRequired = true)]
+		public int Change;
+
+		/// <summary>
+		/// Steps that are part of this job
+		/// </summary>
+		[DataMember(IsRequired = true)]
+		public List<InputJobStep> Steps = new List<InputJobStep>();
+	}
+
+	/// <summary>
+	/// Main class containing input data for the build
+	/// </summary>
+	[DataContract]
+	class InputData
+	{
+		/// <summary>
+		/// Information about builds that have completed.
+		/// </summary>
+		[DataMember]
+		public List<InputJob> Jobs = new List<InputJob>();
+	}
+}
diff --git a/Engine/Source/Programs/BuildAgent/Issues/Issue.cs b/Engine/Source/Programs/BuildAgent/Issues/Issue.cs
new file mode 100644
--- /dev/null
+++ b/Engine/Source/Programs/BuildAgent/Issues/Issue.cs
@@ -0,0 +1,160 @@
+// Copyright 1998-2019 Epic Games, Inc. All Rights Reserved.
+
+using System;
+using System.Collections.Generic;
+using System.Diagnostics;
+using System.Linq;
+using System.Runtime.Serialization;
+using System.Text;
+using System.Threading.Tasks;
+
+namespace BuildAgent.Issues
+{
+	/// <summary>
+	/// Information about a particular issue
+	/// </summary>
+	[DataContract]
+	class Issue
+	{
+		/// <summary>
+		/// The issue id in the database. -1 for issues that have not been posted yet.
+		/// </summary>
+		[DataMember(Order = 0, IsRequired = true)]
+		public long Id = -1;
+
+		/// <summary>
+		/// Project that this belongs to
+		/// </summary>
+		[DataMember(Order = 1, IsRequired = true)]
+		public string Project;
+
+		/// <summary>
+		/// Type common to all diagnostics within this issue.
+		/// </summary>
+		[DataMember(Order = 2, IsRequired = true)]
+		public string Category;
+
+		/// <summary>
+		/// The initial job that this error was seen on. Allows other issues from the same job to be merged.
+		/// </summary>
+		[DataMember(Order = 3, IsRequired = true)]
+		public string InitialJobUrl;
+
+		/// <summary>
+		/// List of strings to display in the details panel for this job
+		/// </summary>
+		[DataMember(Order = 4, IsRequired = true)]
+		public List<IssueDiagnostic> Diagnostics = new List<IssueDiagnostic>();
+
+		/// <summary>
+		/// List of files associated with this issue
+		/// </summary>
+		[DataMember(Order = 5)]
+		public HashSet<string> FileNames = new HashSet<string>();
+
+		/// <summary>
+		/// List of messages associated with this issue
+		/// </summary>
+		[DataMember(Order = 6)]
+		public SortedSet<string> Identifiers = new SortedSet<string>();
+
+		/// <summary>
+		/// Other arbitrary references to assets or files
+		/// </summary>
+		[DataMember(Order = 7)]
+		public SortedSet<string> References = new SortedSet<string>();
+
+		/// <summary>
+		/// The last posted issue summary. Will be updated if it changes.
+		/// </summary>
+		[DataMember(Order = 20)]
+		public string PostedSummary;
+
+		/// <summary>
+		/// Whether we've posted an update to the resolved flag to the server
+		/// </summary>
+		[DataMember(Order = 22)]
+		public bool bPostedResolved;
+
+		/// <summary>
+		/// The time at which the issue was closed. Issues will be retained for 24 hours after they are closed, in case the same issue appears in another stream and to prevent the issue being added again. 
+		/// </summary>
+		[DataMember(Order = 23)]
+		public DateTime? ResolvedAt;
+
+		/// <summary>
+		/// Map of stream name -> step name -> history for builds exhibiting this issue
+		/// </summary>
+		[DataMember(Order = 30, IsRequired = true)]
+		public Dictionary<string, Dictionary<string, IssueHistory>> Streams = new Dictionary<string, Dictionary<string, IssueHistory>>();
+
+		/// <summary>
+		/// Set of changes which may have caused this issue. Used to de-duplicate issues between streams.
+		/// </summary>
+		[DataMember(Order = 31)]
+		public HashSet<int> SourceChanges = new HashSet<int>();
+
+		/// <summary>
+		/// List of possible causers 
+		/// </summary>
+		[DataMember(Order = 32)]
+		public HashSet<string> Watchers = new HashSet<string>();
+
+		/// <summary>
+		/// Set of causers that have yet to be added to the possible causers list
+		/// </summary>
+		[DataMember(Order = 33)]
+		public HashSet<string> PendingWatchers = new HashSet<string>();
+
+		/// <summary>
+		/// Constructor
+		/// </summary>
+		/// <param name="Category">Category of this issue</param>
+		/// <param name="InitialJobUrl">Url of the initial job that this error was seen with</param>
+		/// <param name="ErrorUrl"></param>
+		public Issue(string Project, string Category, string InitialJobUrl, IssueDiagnostic Diagnostic)
+		{
+			this.Project = Project;
+			this.Category = Category;
+			this.InitialJobUrl = InitialJobUrl;
+			this.Diagnostics.Add(Diagnostic);
+		}
+
+		/// <summary>
+		/// Determines whether the issue can be closed
+		/// </summary>
+		/// <returns>True if the issue can be closed</returns>
+		public bool IsResolved()
+		{
+			return Streams.Values.All(x => x.Values.All(y => y.NextSuccessfulBuild != null));
+		}
+
+		/// <summary>
+		/// Finds all the steps which are related to this issue
+		/// </summary>
+		/// <returns>Set of step names</returns>
+		public SortedSet<string> GetStepNames()
+		{
+			return new SortedSet<string>(Streams.Values.SelectMany(x => x.Keys));
+		}
+
+		/// <summary>
+		/// Format the issue for the debugger
+		/// </summary>
+		/// <returns>String representation of the issue</returns>
+		public override string ToString()
+		{
+			StringBuilder Result = new StringBuilder();
+			if(Id == -1)
+			{
+				Result.Append("[New] ");
+			}
+			else
+			{
+				Result.AppendFormat("[{0}] ", Id);
+			}
+			Result.AppendFormat("Errors in {0}", String.Join(", ", GetStepNames()));
+			return Result.ToString();
+		}
+	}
+}
diff --git a/Engine/Source/Programs/BuildAgent/Issues/IssueBuild.cs b/Engine/Source/Programs/BuildAgent/Issues/IssueBuild.cs
new file mode 100644
--- /dev/null
+++ b/Engine/Source/Programs/BuildAgent/Issues/IssueBuild.cs
@@ -0,0 +1,100 @@
+// Copyright 1998-2019 Epic Games, Inc. All Rights Reserved.
+
+using System;
+using System.Collections.Generic;
+using System.Diagnostics;
+using System.Linq;
+using System.Runtime.Serialization;
+using System.Text;
+using System.Threading.Tasks;
+
+namespace BuildAgent.Issues
+{
+	/// <summary>
+	/// Represents information about a build associated with a particular issue
+	/// </summary>
+	[DataContract]
+	[DebuggerDisplay("{Change}: {JobName}")]
+	class IssueBuild : IComparable<IssueBuild>
+	{
+		/// <summary>
+		/// The id of this build on the server
+		/// </summary>
+		[DataMember(IsRequired = true)]
+		public long Id = -1;
+
+		/// <summary>
+		/// The changelist that this build was run at.
+		/// </summary>
+		[DataMember(IsRequired = true)]
+		public int Change;
+
+		/// <summary>
+		/// Name of this job
+		/// </summary>
+		[DataMember(IsRequired = true)]
+		public string JobName;
+
+		/// <summary>
+		/// Url for this job
+		/// </summary>
+		[DataMember(IsRequired = true)]
+		public string JobUrl;
+
+		/// <summary>
+		/// Name of this job step (typically null for sentinel builds)
+		/// </summary>
+		[DataMember(IsRequired = true)]
+		public string JobStepName;
+
+		/// <summary>
+		/// Url for this job step (typically null for sentinel builds)
+		/// </summary>
+		[DataMember(IsRequired = true)]
+		public string JobStepUrl;
+
+		/// <summary>
+		/// Url of the first error within this job
+		/// </summary>
+		[DataMember]
+		public string ErrorUrl;
+
+		/// <summary>
+		/// Whether this build has been posted to the server or not
+		/// </summary>
+		[DataMember]
+		public bool bPostedToServer = false;
+
+		/// <summary>
+		/// Constructor
+		/// </summary>
+		public IssueBuild(int Change, string JobName, string JobUrl, string JobStepName, string JobStepUrl, string ErrorUrl)
+		{
+			this.Change = Change;
+			this.JobName = JobName;
+			this.JobUrl = JobUrl;
+			this.JobStepName = JobStepName;
+			this.JobStepUrl = JobStepUrl;
+			this.ErrorUrl = ErrorUrl;
+		}
+
+		/// <summary>
+		/// Compares this build to another build
+		/// </summary>
+		/// <param name="Other">Build to compare to</param>
+		/// <returns>Value indicating how the two builds should be ordered</returns>
+		public int CompareTo(IssueBuild Other)
+		{
+			int Delta = Change - Other.Change;
+			if (Delta == 0)
+			{
+				Delta = JobUrl.CompareTo(Other.JobUrl);
+				if(Delta == 0)
+				{
+					Delta = JobStepUrl.CompareTo(Other.JobStepUrl);
+				}
+			}
+			return Delta;
+		}
+	}
+}
diff --git a/Engine/Source/Programs/BuildAgent/Issues/IssueDiagnostic.cs b/Engine/Source/Programs/BuildAgent/Issues/IssueDiagnostic.cs
new file mode 100644
--- /dev/null
+++ b/Engine/Source/Programs/BuildAgent/Issues/IssueDiagnostic.cs
@@ -0,0 +1,63 @@
+using System;
+using System.Collections.Generic;
+using System.Diagnostics;
+using System.Linq;
+using System.Runtime.Serialization;
+using System.Text;
+using System.Threading.Tasks;
+
+namespace BuildAgent.Issues
+{
+	/// <summary>
+	/// Information about a message
+	/// </summary>
+	[DataContract]
+	[DebuggerDisplay("{Message}")]
+	public class IssueDiagnostic
+	{
+		/// <summary>
+		/// Name of this job step
+		/// </summary>
+		[DataMember(Order = 0, IsRequired = true)]
+		public string JobStepName;
+
+		/// <summary>
+		/// The url of this job step
+		/// </summary>
+		[DataMember(Order = 1)]
+		public string JobStepUrl;
+
+		/// <summary>
+		/// The diagnostic message
+		/// </summary>
+		[DataMember(Order = 2, IsRequired = true)]
+		public string Message;
+
+		/// <summary>
+		/// Url to the error
+		/// </summary>
+		[DataMember(Order = 3, IsRequired = true)]
+		public string ErrorUrl;
+
+		/// <summary>
+		/// Whether or not this diagnostic has been posted to the server
+		/// </summary>
+		[DataMember(Order = 4)]
+		public bool bPostedToServer = false;
+
+		/// <summary>
+		/// Constructor
+		/// </summary>
+		/// <param name="JobStepName">Name of the job step</param>
+		/// <param name="JobStepUrl">Url of the job step</param>
+		/// <param name="Message">Message to display</param>
+		/// <param name="ErrorUrl">Url to link to</param>
+		public IssueDiagnostic(string JobStepName, string JobStepUrl, string Message, string ErrorUrl)
+		{
+			this.JobStepName = JobStepName;
+			this.JobStepUrl = JobStepUrl;
+			this.Message = Message;
+			this.ErrorUrl = ErrorUrl;
+		}
+	}
+}
diff --git a/Engine/Source/Programs/BuildAgent/Issues/IssueHistory.cs b/Engine/Source/Programs/BuildAgent/Issues/IssueHistory.cs
new file mode 100644
--- /dev/null
+++ b/Engine/Source/Programs/BuildAgent/Issues/IssueHistory.cs
@@ -0,0 +1,101 @@
+// Copyright 1998-2019 Epic Games, Inc. All Rights Reserved.
+
+using System;
+using System.Collections.Generic;
+using System.Linq;
+using System.Runtime.Serialization;
+using System.Text;
+using System.Threading.Tasks;
+
+namespace BuildAgent.Issues
+{
+	/// <summary>
+	/// History of builds within a particular stream that contribute to an issue
+	/// </summary>
+	[DataContract]
+	class IssueHistory
+	{
+		/// <summary>
+		/// The previous build before it started failing. This should be updated as new builds come in.
+		/// </summary>
+		[DataMember]
+		public IssueBuild PrevSuccessfulBuild;
+
+		/// <summary>
+		/// List of failing builds contributing to this issue
+		/// </summary>
+		[DataMember]
+		public List<IssueBuild> FailedBuilds = new List<IssueBuild>();
+
+		/// <summary>
+		/// The first successful build after the failures.
+		/// </summary>
+		[DataMember]
+		public IssueBuild NextSuccessfulBuild;
+
+		/// <summary>
+		/// Constructs a new history for a particular stream
+		/// </summary>
+		public IssueHistory(IssueBuild PrevSuccessfulBuild)
+		{
+			this.PrevSuccessfulBuild  = PrevSuccessfulBuild;
+		}
+
+		/// <summary>
+		/// Adds a failed build to this object
+		/// </summary>
+		/// <param name="Build">The failed build</param>
+		public void AddFailedBuild(IssueBuild Build)
+		{
+			int Index = FailedBuilds.BinarySearch(Build);
+			if (Index < 0)
+			{
+				FailedBuilds.Insert(~Index, Build);
+			}
+		}
+
+		/// <summary>
+		/// Determines whether a given build can be added to an issue. This filters cases where an issue does not already have builds for the given stream, or where there is a successful build between the new build and known failures for this issue.
+		/// </summary>
+		/// <param name="Build">The build to add</param>
+		public bool CanAddFailedBuild(int Change)
+		{
+			// Check that this build is not after a succesful build
+			if(NextSuccessfulBuild != null && Change >= NextSuccessfulBuild.Change)
+			{
+				return false;
+			}
+
+			// Check that this build is not before the last known successful build
+			if(PrevSuccessfulBuild != null && Change <= PrevSuccessfulBuild.Change)
+			{
+				return false;
+			}
+
+			// Otherwise allow it
+			return true;
+		}
+
+		/// <summary>
+		/// Enumerates all the builds in this stream
+		/// </summary>
+		public IEnumerable<IssueBuild> Builds
+		{
+			get
+			{
+				if(PrevSuccessfulBuild != null)
+				{
+					yield return PrevSuccessfulBuild;
+				}
+				foreach(IssueBuild FailedBuild in FailedBuilds)
+				{
+					yield return FailedBuild;
+				}
+				if(NextSuccessfulBuild != null)
+				{
+					yield return NextSuccessfulBuild;
+				}
+			}
+		}
+	}
+}
diff --git a/Engine/Source/Programs/BuildAgent/Issues/Matcher.cs b/Engine/Source/Programs/BuildAgent/Issues/Matcher.cs
new file mode 100644
--- /dev/null
+++ b/Engine/Source/Programs/BuildAgent/Issues/Matcher.cs
@@ -0,0 +1,314 @@
+// Copyright 1998-2019 Epic Games, Inc. All Rights Reserved.
+
+using System;
+using System.Collections.Generic;
+using System.Linq;
+using System.Text;
+using System.Text.RegularExpressions;
+using System.Threading.Tasks;
+using Tools.DotNETCommon.Perforce;
+
+namespace BuildAgent.Issues
+{
+	/// <summary>
+	/// Class that implements a pattern matcher for a particular class of errors
+	/// </summary>
+	abstract class Matcher
+	{
+		/// <summary>
+		/// The category name
+		/// </summary>
+		public abstract string Category
+		{
+			get;
+		}
+
+		/// <summary>
+		/// Creates fingerprints from any matching diagnostics
+		/// </summary>
+		/// <param name="Job">The job that was run</param>
+		/// <param name="JobStep">The job step that was run</param>
+		/// <param name="Diagnostics">List of diagnostics that were produced by the build. Items should be removed from this list if they match.</param>
+		/// <param name="Issues">List which receives all the matched issues.</param>
+		public virtual void Match(InputJob Job, InputJobStep JobStep, List<InputDiagnostic> Diagnostics, List<Issue> Issues)
+		{
+			for (int Idx = 0; Idx < Diagnostics.Count; Idx++)
+			{
+				InputDiagnostic Diagnostic = Diagnostics[Idx];
+				if(TryMatch(Job, JobStep, Diagnostic, Issues))
+				{
+					Diagnostics.RemoveAt(Idx);
+					Idx--;
+				}
+			}
+		}
+
+		/// <summary>
+		/// Tries to create a fingerprint from an individual diagnostic.
+		/// </summary>
+		/// <param name="Job">The job that was run</param>
+		/// <param name="JobStep">The job step that was run</param>
+		/// <param name="Diagnostic">A diagnostic from the given job step</param>
+		/// <param name="Issues">List which receives all the matched issues.</param>
+		/// <returns>True if this diagnostic should be removed (usually because a fingerprint was created)</returns>
+		public abstract bool TryMatch(InputJob Job, InputJobStep JobStep, InputDiagnostic Diagnostic, List<Issue> Issues);
+
+		/// <summary>
+		/// Determines if one issue can be merged into another
+		/// </summary>
+		/// <param name="Source">The source issue</param>
+		/// <param name="Target">The target issue</param>
+		public virtual bool CanMerge(Issue Source, Issue Target)
+		{
+			// Make sure the categories match
+			if (Source.Category != Target.Category)
+			{
+				return false;
+			}
+
+			// Check that a filename or message matches
+			if (!Source.FileNames.Any(x => Target.FileNames.Contains(x)) && !Source.Identifiers.Any(x => Target.Identifiers.Contains(x)))
+			{
+				return false;
+			}
+
+			return true;
+		}
+
+		/// <summary>
+		/// Determines if an issue can be merged into another issue that occurred at the same initial job
+		/// </summary>
+		/// <param name="Source">The source issue</param>
+		/// <param name="Target">The target issue</param>
+		/// <returns>True if the two new issues can be merged</returns>
+		public virtual bool CanMergeInitialJob(Issue Source, Issue Target)
+		{
+			return Source.Category == Target.Category;
+		}
+
+		/// <summary>
+		/// Merge one fingerprint with another
+		/// </summary>
+		/// <param name="Source">The source fingerprint</param>
+		/// <param name="Target">The fingerprint to merge into</param>
+		public virtual void Merge(Issue Source, Issue Target)
+		{
+			HashSet<string> TargetMessages = new HashSet<string>(Target.Diagnostics.Select(x => x.Message), StringComparer.Ordinal);
+			foreach(IssueDiagnostic SourceDiagnostic in Source.Diagnostics)
+			{
+				if(Target.Diagnostics.Count >= 50)
+				{
+					break;
+				}
+				if(!TargetMessages.Contains(SourceDiagnostic.Message))
+				{
+					Target.Diagnostics.Add(SourceDiagnostic);
+				}
+			}
+
+			Target.FileNames.UnionWith(Source.FileNames);
+			Target.Identifiers.UnionWith(Source.Identifiers);
+			Target.References.UnionWith(Source.References);
+		}
+
+		/// <summary>
+		/// Filters all the likely causers from the list of changes since an issue was created
+		/// </summary>
+		/// <param name="Perforce">The perforce connection</param>
+		/// <param name="Issue">The build issue</param>
+		/// <param name="Changes">List of changes since the issue first occurred.</param>
+		/// <returns>List of changes which are causers for the issue</returns>
+		public virtual List<ChangeInfo> FindCausers(PerforceConnection Perforce, Issue Issue, IReadOnlyList<ChangeInfo> Changes)
+		{
+			List<ChangeInfo> Causers = new List<ChangeInfo>();
+
+			SortedSet<string> FileNamesWithoutPath = GetFileNamesWithoutPath(Issue.FileNames);
+			if (FileNamesWithoutPath.Count > 0)
+			{
+				foreach (ChangeInfo Change in Changes)
+				{
+					DescribeRecord DescribeRecord = GetDescribeRecord(Perforce, Change);
+					if (ContainsFileNames(DescribeRecord, FileNamesWithoutPath))
+					{
+						Causers.Add(Change);
+					}
+				}
+			}
+
+			if(Causers.Count > 0)
+			{
+				return Causers;
+			}
+			else
+			{
+				return new List<ChangeInfo>(Changes);
+			}
+		}
+
+		/// <summary>
+		/// Utility method to get the describe record for a change. Caches it on the ChangeInfo object as necessary.
+		/// </summary>
+		/// <param name="Perforce">The Perforce connection</param>
+		/// <param name="Change">The change to query</param>
+		public DescribeRecord GetDescribeRecord(PerforceConnection Perforce, ChangeInfo Change)
+		{
+			if(Change.CachedDescribeRecord == null)
+			{
+				Change.CachedDescribeRecord = Perforce.Describe(Change.Record.Number).Data;
+			}
+			return Change.CachedDescribeRecord;
+		}
+
+		/// <summary>
+		/// Tests whether a change is a code change
+		/// </summary>
+		/// <param name="Perforce">The Perforce connection</param>
+		/// <param name="Change">The change to query</param>
+		/// <returns>True if the change is a code change</returns>
+		public bool ContainsAnyFileWithExtension(PerforceConnection Perforce, ChangeInfo Change, string[] Extensions)
+		{
+			DescribeRecord Record = GetDescribeRecord(Perforce, Change);
+			foreach(DescribeFileRecord File in Record.Files)
+			{
+				foreach(string Extension in Extensions)
+				{
+					if(File.DepotFile.EndsWith(Extension, StringComparison.OrdinalIgnoreCase))
+					{
+						return true;
+					}
+				}
+			}
+			return false;
+		}
+
+		/// <summary>
+		/// Determines if this change is a likely causer for an issue
+		/// </summary>
+		/// <param name="DescribeRecord">The change describe record</param>
+		/// <param name="Fingerprint">Fingerprint for the issue</param>
+		/// <returns>True if the change is a likely culprit</returns>
+		protected static bool ContainsFileNames(DescribeRecord DescribeRecord, SortedSet<string> FileNamesWithoutPath)
+		{
+			foreach (DescribeFileRecord File in DescribeRecord.Files)
+			{
+				int Idx = File.DepotFile.LastIndexOf('/');
+				if (Idx != -1)
+				{
+					string FileName = File.DepotFile.Substring(Idx + 1);
+					if (FileNamesWithoutPath.Contains(FileName))
+					{
+						return true;
+					}
+				}
+			}
+			return false;
+		}
+
+		/// <summary>
+		/// Normalizes a filename to a path within the workspace
+		/// </summary>
+		/// <param name="FileName">Filename to normalize</param>
+		/// <param name="BaseDirectory">Base directory containing the workspace</param>
+		/// <returns>Normalized filename</returns>
+		protected string GetNormalizedFileName(string FileName, string BaseDirectory)
+		{
+			string NormalizedFileName = FileName.Replace('\\', '/');
+			if (!String.IsNullOrEmpty(BaseDirectory))
+			{
+				// Normalize the expected base directory for errors in this build, and attempt to strip it from the file name
+				string NormalizedBaseDirectory = BaseDirectory;
+				if (NormalizedBaseDirectory != null && NormalizedBaseDirectory.Length > 0)
+				{
+					NormalizedBaseDirectory = NormalizedBaseDirectory.Replace('\\', '/').TrimEnd('/') + "/";
+				}
+				if (NormalizedFileName.StartsWith(NormalizedBaseDirectory, StringComparison.OrdinalIgnoreCase))
+				{
+					NormalizedFileName = NormalizedFileName.Substring(NormalizedBaseDirectory.Length);
+				}
+			}
+			else
+			{
+				// Try to match anything under a 'Sync' folder.
+				Match FallbackRegex = Regex.Match(NormalizedFileName, "/Sync/(.*)");
+				if (FallbackRegex.Success)
+				{
+					NormalizedFileName = FallbackRegex.Groups[1].Value;
+				}
+			}
+			return NormalizedFileName;
+		}
+
+		/// <summary>
+		/// Finds all the unique filenames without their path components
+		/// </summary>
+		/// <returns>Set of sorted filenames</returns>
+		protected static SortedSet<string> GetFileNamesWithoutPath(IEnumerable<string> FileNames)
+		{
+			SortedSet<string> FileNamesWithoutPath = new SortedSet<string>(StringComparer.OrdinalIgnoreCase);
+			foreach (string FileName in FileNames)
+			{
+				int Idx = FileName.LastIndexOf('/');
+				if (Idx != -1)
+				{
+					FileNamesWithoutPath.Add(FileName.Substring(Idx + 1));
+				}
+			}
+			return FileNamesWithoutPath;
+		}
+
+		/// <summary>
+		/// Gets a set of unique source file names that relate to this issue
+		/// </summary>
+		/// <returns>Set of source file names</returns>
+		protected static SortedSet<string> GetSourceFileNames(IEnumerable<string> FileNames)
+		{
+			SortedSet<string> ShortFileNames = new SortedSet<string>(StringComparer.OrdinalIgnoreCase);
+			foreach (string FileName in FileNames)
+			{
+				int Idx = FileName.LastIndexOfAny(new char[] { '/', '\\' });
+				if (Idx != -1)
+				{
+					string ShortFileName = FileName.Substring(Idx + 1);
+					if (!ShortFileName.StartsWith("Module.", StringComparison.OrdinalIgnoreCase))
+					{
+						ShortFileNames.Add(ShortFileName);
+					}
+				}
+			}
+			return ShortFileNames;
+		}
+
+		/// <summary>
+		/// Gets a set of unique asset filenames that relate to this issue
+		/// </summary>
+		/// <returns>Set of asset names</returns>
+		protected static SortedSet<string> GetAssetNames(IEnumerable<string> FileNames)
+		{
+			SortedSet<string> ShortFileNames = new SortedSet<string>(StringComparer.OrdinalIgnoreCase);
+			foreach (string FileName in FileNames)
+			{
+				int Idx = FileName.LastIndexOfAny(new char[] { '/', '\\' });
+				if (Idx != -1)
+				{
+					string AssetName = FileName.Substring(Idx + 1);
+
+					int DotIdx = AssetName.LastIndexOf('.');
+					if (DotIdx != -1)
+					{
+						AssetName = AssetName.Substring(0, DotIdx);
+					}
+
+					ShortFileNames.Add(AssetName);
+				}
+			}
+			return ShortFileNames;
+		}
+
+		/// <summary>
+		/// Gets the summary for an issue
+		/// </summary>
+		/// <param name="Issue">The issue to summarize</param>
+		/// <returns>The summary text for this issue</returns>
+		public abstract string GetSummary(Issue Issue);
+	}
+}
diff --git a/Engine/Source/Programs/BuildAgent/Issues/Matchers/CompileIssueMatcher.cs b/Engine/Source/Programs/BuildAgent/Issues/Matchers/CompileIssueMatcher.cs
new file mode 100644
--- /dev/null
+++ b/Engine/Source/Programs/BuildAgent/Issues/Matchers/CompileIssueMatcher.cs
@@ -0,0 +1,91 @@
+// Copyright 1998-2019 Epic Games, Inc. All Rights Reserved.
+
+using System;
+using System.Collections.Generic;
+using System.Linq;
+using System.Text;
+using System.Text.RegularExpressions;
+using System.Threading.Tasks;
+
+namespace BuildAgent.Issues.Matchers
+{
+	class CompileIssueMatcher : GenericCodeIssueMatcher
+	{
+		static string[] SourceFileExtensions =
+		{
+			".cpp",
+			".h",
+			".cc",
+			".hh",
+			".m",
+			".mm",
+			".rc",
+			".inl",
+			".inc"
+		};
+
+		public override string Category => "Compile";
+
+		public override bool TryMatch(InputJob Job, InputJobStep JobStep, InputDiagnostic Diagnostic, List<Issue> Issues)
+		{
+			// Find a list of source files with errors
+			HashSet<string> ErrorFileNames = new HashSet<string>(StringComparer.OrdinalIgnoreCase);
+			foreach (Match FileMatch in Regex.Matches(Diagnostic.Message, @"^\s*((?:[A-Za-z]:)?[^\s(:]+)[\(:]\d[\s\d:\)]+(?:warning|error|fatal error)", RegexOptions.Multiline))
+			{
+				if (FileMatch.Success)
+				{
+					string FileName = GetNormalizedFileName(FileMatch.Groups[1].Value, JobStep.BaseDirectory);
+					if (SourceFileExtensions.Any(x => FileName.EndsWith(x, StringComparison.OrdinalIgnoreCase)))
+					{
+						ErrorFileNames.Add(FileName);
+					}
+				}
+			}
+
+			// Find any referenced files in compiler output format
+			HashSet<string> ReferencedFileNames = new HashSet<string>(StringComparer.OrdinalIgnoreCase);
+			foreach (Match FileMatch in Regex.Matches(Diagnostic.Message, @"^\s*(?:In file included from\s*)?((?:[A-Za-z]:)?[^\s(:]+)[\(:]\d", RegexOptions.Multiline))
+			{
+				if (FileMatch.Success)
+				{
+					string FileName = GetNormalizedFileName(FileMatch.Groups[1].Value, JobStep.BaseDirectory);
+					if (SourceFileExtensions.Any(x => FileName.EndsWith(x, StringComparison.OrdinalIgnoreCase)))
+					{
+						ReferencedFileNames.Add(FileName);
+					}
+				}
+			}
+
+			// If we found any source files, create a diagnostic category for them
+			if (ReferencedFileNames.Count > 0)
+			{
+				Issue Issue = new Issue(Job.Project, Category, Job.Url, new IssueDiagnostic(JobStep.Name, JobStep.Url, ShortenPaths(Diagnostic.Message), Diagnostic.Url));
+				Issue.FileNames.UnionWith(ReferencedFileNames);
+				Issue.Identifiers.UnionWith(GetSourceFileNames(ErrorFileNames));
+				Issues.Add(Issue);
+				return true;
+			}
+
+			// Otherwise pass
+			return false;
+		}
+
+		static string ShortenPaths(string Text)
+		{
+			Text = Regex.Replace(Text, @"^(\s*(?:In file included from\s*)?)(?:[A-Za-z]:)?[^\s(:]+[/\\]([^\s/\\]+[\(:]\d)", "$1$2", RegexOptions.Multiline);
+			return Text;
+		}
+
+		public override string GetSummary(Issue Issue)
+		{
+			if (Issue.Identifiers.Count == 0)
+			{
+				return "Compile errors";
+			}
+			else
+			{
+				return String.Format("Compile errors in {0}", String.Join(", ", Issue.Identifiers));
+			}
+		}
+	}
+}
diff --git a/Engine/Source/Programs/BuildAgent/Issues/Matchers/ContentIssueMatcher.cs b/Engine/Source/Programs/BuildAgent/Issues/Matchers/ContentIssueMatcher.cs
new file mode 100644
--- /dev/null
+++ b/Engine/Source/Programs/BuildAgent/Issues/Matchers/ContentIssueMatcher.cs
@@ -0,0 +1,47 @@
+// Copyright 1998-2019 Epic Games, Inc. All Rights Reserved.
+
+using System;
+using System.Collections.Generic;
+using System.Linq;
+using System.Text;
+using System.Text.RegularExpressions;
+using System.Threading.Tasks;
+using Tools.DotNETCommon.Perforce;
+
+namespace BuildAgent.Issues.Matchers
+{
+	class ContentIssueMatcher : GenericContentMatcher
+	{
+		public override string Category => "Content";
+
+		public override bool TryMatch(InputJob Job, InputJobStep JobStep, InputDiagnostic Diagnostic, List<Issue> Issues)
+		{
+			HashSet<string> FileNames = new HashSet<string>();
+			foreach(Match Match in Regex.Matches(Diagnostic.Message, @"^\s*Log[a-zA-Z0-9]+:\s+(?:Error:|Warning:)\s+((?:[a-zA-Z]:)?[^:]+(?:.uasset|.umap)):\s*(.*)"))
+			{
+				FileNames.Add(GetNormalizedFileName(Match.Groups[1].Value, JobStep.BaseDirectory));
+			}
+
+			if(FileNames.Count > 0)
+			{
+				Issue Issue = new Issue(Job.Project, Category, Job.Url, new IssueDiagnostic(JobStep.Name, JobStep.Url, Diagnostic.Message, Diagnostic.Url));
+				Issue.FileNames.UnionWith(FileNames);
+				Issues.Add(Issue);
+				return true;
+			}
+			return false;
+		}
+
+		public override string GetSummary(Issue Issue)
+		{
+			if (Issue.FileNames.Count == 0)
+			{
+				return "Content errors";
+			}
+			else
+			{
+				return String.Format("Content errors in {0}", String.Join(", ", GetFileNamesWithoutPath(Issue.FileNames)));
+			}
+		}
+	}
+}
diff --git a/Engine/Source/Programs/BuildAgent/Issues/Matchers/CopyrightNoticeMatcher.cs b/Engine/Source/Programs/BuildAgent/Issues/Matchers/CopyrightNoticeMatcher.cs
new file mode 100644
--- /dev/null
+++ b/Engine/Source/Programs/BuildAgent/Issues/Matchers/CopyrightNoticeMatcher.cs
@@ -0,0 +1,61 @@
+// Copyright 1998-2019 Epic Games, Inc. All Rights Reserved.
+
+using System;
+using System.Collections.Generic;
+using System.Linq;
+using System.Text;
+using System.Text.RegularExpressions;
+using System.Threading.Tasks;
+
+namespace BuildAgent.Issues.Matchers
+{
+	class CopyrightNoticeMatcher : GenericCodeIssueMatcher
+	{
+		public override string Category => "Copyright";
+
+		public override bool TryMatch(InputJob Job, InputJobStep JobStep, InputDiagnostic Diagnostic, List<Issue> Issues)
+		{
+			// Make sure we're running a step that this applies to
+			if(JobStep.Name.IndexOf("Copyright", StringComparison.OrdinalIgnoreCase) == -1)
+			{
+				return false;
+			}
+
+			// Find any files in compiler output format
+			HashSet<string> SourceFileNames = new HashSet<string>(StringComparer.OrdinalIgnoreCase);
+			foreach (Match FileMatch in Regex.Matches(Diagnostic.Message, @"^\s*(?:WARNING|ERROR):\s*([^ ]+\.[a-zA-Z]+):", RegexOptions.Multiline))
+			{
+				if (FileMatch.Success)
+				{
+					string SourceFileName = FileMatch.Groups[1].Value.Replace('\\', '/');
+					SourceFileNames.Add(SourceFileName);
+				}
+			}
+
+			// If we found any source files, create a diagnostic category for them
+			if (SourceFileNames.Count > 0)
+			{
+				Issue Issue = new Issue(Job.Project, Category, Job.Url, new IssueDiagnostic(JobStep.Name, JobStep.Url, Diagnostic.Message, Diagnostic.Url));
+				Issue.FileNames.UnionWith(SourceFileNames);
+				Issues.Add(Issue);
+				return true;
+			}
+
+			// Otherwise pass
+			return false;
+		}
+
+		public override string GetSummary(Issue Issue)
+		{
+			SortedSet<string> ShortFileNames = GetSourceFileNames(Issue.FileNames);
+			if (ShortFileNames.Count == 0)
+			{
+				return "Missing copyright notices";
+			}
+			else
+			{
+				return String.Format("Missing copyright in {0}", String.Join(", ", ShortFileNames));
+			}
+		}
+	}
+}
diff --git a/Engine/Source/Programs/BuildAgent/Issues/Matchers/DefaultIssueMatcher.cs b/Engine/Source/Programs/BuildAgent/Issues/Matchers/DefaultIssueMatcher.cs
new file mode 100644
--- /dev/null
+++ b/Engine/Source/Programs/BuildAgent/Issues/Matchers/DefaultIssueMatcher.cs
@@ -0,0 +1,46 @@
+// Copyright 1998-2019 Epic Games, Inc. All Rights Reserved.
+
+using System;
+using System.Collections.Generic;
+using System.Linq;
+using System.Text;
+using System.Text.RegularExpressions;
+using System.Threading.Tasks;
+using Tools.DotNETCommon.Perforce;
+
+namespace BuildAgent.Issues.Matchers
+{
+	class DefaultIssueMatcher : Matcher
+	{
+		public override string Category => "Default";
+
+		public override bool TryMatch(InputJob Job, InputJobStep JobStep, InputDiagnostic Diagnostic, List<Issue> Issues)
+		{
+			string DefaultProject = String.Format("{0} (Unmatched)", Job.Project);
+
+			Issue Issue = new Issue(DefaultProject, Category, Job.Url, new IssueDiagnostic(JobStep.Name, JobStep.Url, Diagnostic.Message, Diagnostic.Url));
+			Issue.Identifiers.Add(Diagnostic.Message);
+			Issues.Add(Issue);
+			
+			return true;
+		}
+
+		public override List<ChangeInfo> FindCausers(PerforceConnection Perforce, Issue Issue, IReadOnlyList<ChangeInfo> Changes)
+		{
+			return new List<ChangeInfo>();
+		}
+
+		public override string GetSummary(Issue Issue)
+		{
+			string Message = Issue.Diagnostics[0].Message;
+
+			int Idx = Message.IndexOf('\n');
+			if(Idx >= 0)
+			{
+				Message = Message.Substring(0, Idx).Trim();
+			}
+
+			return Message;
+		}
+	}
+}
diff --git a/Engine/Source/Programs/BuildAgent/Issues/Matchers/GenericCodeIssueMatcher.cs b/Engine/Source/Programs/BuildAgent/Issues/Matchers/GenericCodeIssueMatcher.cs
new file mode 100644
--- /dev/null
+++ b/Engine/Source/Programs/BuildAgent/Issues/Matchers/GenericCodeIssueMatcher.cs
@@ -0,0 +1,51 @@
+// Copyright 1998-2019 Epic Games, Inc. All Rights Reserved.
+
+using System;
+using System.Collections.Generic;
+using System.Linq;
+using System.Text;
+using System.Threading.Tasks;
+using Tools.DotNETCommon.Perforce;
+
+namespace BuildAgent.Issues.Matchers
+{
+	/// <summary>
+	/// Base class for pattern matchers that match errors in source code
+	/// </summary>
+	abstract class GenericCodeIssueMatcher : Matcher
+	{
+		/// <summary>
+		/// Set of extensions to treat as code
+		/// </summary>
+		static readonly string[] CodeExtensions =
+		{
+			".c",
+			".cc",
+			".cpp",
+			".m",
+			".mm",
+			".rc",
+			".cs",
+			".h",
+			".hpp",
+			".inl",
+			".uproject",
+			".uplugin",
+			".ini",
+		};
+
+		/// <summary>
+		/// Finds any suspected causers for a particular failure. Excludes any changes that don't contain code.
+		/// </summary>
+		/// <param name="Perforce">The perforce connection</param>
+		/// <param name="Issue">The build issue</param>
+		/// <param name="Changes">List of changes since the issue first occurred.</param>
+		/// <returns>List of changes which are causers for the issue</returns>
+		public override List<ChangeInfo> FindCausers(PerforceConnection Perforce, Issue Issue, IReadOnlyList<ChangeInfo> Changes)
+		{
+			List<ChangeInfo> Causers = base.FindCausers(Perforce, Issue, Changes);
+			Causers.RemoveAll(x => !ContainsAnyFileWithExtension(Perforce, x, CodeExtensions));
+			return Causers;
+		}
+	}
+}
diff --git a/Engine/Source/Programs/BuildAgent/Issues/Matchers/GenericContentIssueMatcher.cs b/Engine/Source/Programs/BuildAgent/Issues/Matchers/GenericContentIssueMatcher.cs
new file mode 100644
--- /dev/null
+++ b/Engine/Source/Programs/BuildAgent/Issues/Matchers/GenericContentIssueMatcher.cs
@@ -0,0 +1,41 @@
+// Copyright 1998-2019 Epic Games, Inc. All Rights Reserved.
+
+using System;
+using System.Collections.Generic;
+using System.Linq;
+using System.Text;
+using System.Threading.Tasks;
+using Tools.DotNETCommon.Perforce;
+
+namespace BuildAgent.Issues.Matchers
+{
+	/// <summary>
+	/// Base class for pattern matchers that match errors in source code
+	/// </summary>
+	abstract class GenericContentMatcher : Matcher
+	{
+		/// <summary>
+		/// Set of extensions to treat as code
+		/// </summary>
+		static readonly string[] ContentExtensions =
+		{
+			".uasset",
+			".umap",
+			".ini",
+		};
+
+		/// <summary>
+		/// Finds any suspected causers for a particular failure. Excludes any changes that don't contain code.
+		/// </summary>
+		/// <param name="Perforce">The perforce connection</param>
+		/// <param name="Issue">The build issue</param>
+		/// <param name="Changes">List of changes since the issue first occurred.</param>
+		/// <returns>List of changes which are causers for the issue</returns>
+		public override List<ChangeInfo> FindCausers(PerforceConnection Perforce, Issue Issue, IReadOnlyList<ChangeInfo> Changes)
+		{
+			List<ChangeInfo> Causers = base.FindCausers(Perforce, Issue, Changes);
+			Causers.RemoveAll(x => !ContainsAnyFileWithExtension(Perforce, x, ContentExtensions));
+			return Causers;
+		}
+	}
+}
diff --git a/Engine/Source/Programs/BuildAgent/Issues/Matchers/UndefinedSymbolMatcher.cs b/Engine/Source/Programs/BuildAgent/Issues/Matchers/UndefinedSymbolMatcher.cs
new file mode 100644
--- /dev/null
+++ b/Engine/Source/Programs/BuildAgent/Issues/Matchers/UndefinedSymbolMatcher.cs
@@ -0,0 +1,146 @@
+// Copyright 1998-2019 Epic Games, Inc. All Rights Reserved.
+
+using System;
+using System.Collections.Generic;
+using System.Linq;
+using System.Text;
+using System.Text.RegularExpressions;
+using System.Threading.Tasks;
+using Tools.DotNETCommon.Perforce;
+
+namespace BuildAgent.Issues.Matchers
+{
+	class UndefinedSymbolMatcher : GenericCodeIssueMatcher
+	{
+		public override string Category => "UndefinedSymbol";
+
+		public override bool TryMatch(InputJob Job, InputJobStep JobStep, InputDiagnostic Diagnostic, List<Issue> Issues)
+		{
+			List<string> SymbolMatches = new List<string>();
+
+			// Mac link error:
+			//   Undefined symbols for architecture arm64:
+            //     "Foo::Bar() const", referenced from:
+			if (Regex.IsMatch(Diagnostic.Message, "^Undefined symbols"))
+			{
+				foreach(string Line in Diagnostic.Message.Split('\n'))
+				{
+					Match SymbolMatch = Regex.Match(Line, "^  \"(.+)\"");
+					if (SymbolMatch.Success)
+					{
+						SymbolMatches.Add(SymbolMatch.Groups[1].Value);
+					}
+				}
+			}
+
+			// Android link error:
+			//   Foo.o:(.data.rel.ro + 0x5d88): undefined reference to `Foo::Bar()'
+			Match UndefinedReference = Regex.Match(Diagnostic.Message, ": undefined reference to [`']([^`']+)");
+			if(UndefinedReference.Success)
+			{
+				SymbolMatches.Add(UndefinedReference.Groups[1].Value);
+			}
+
+			// LLD link error:
+			//   ld.lld.exe: error: undefined symbol: Foo::Bar() const
+			Match LldMatch = Regex.Match(Diagnostic.Message, "error: undefined symbol:\\s*(.+)");
+			if (LldMatch.Success)
+			{
+				SymbolMatches.Add(LldMatch.Groups[1].Value);
+			}
+
+			// Link error:
+			//   Link: error: L0039: reference to undefined symbol `Foo::Bar() const' in file
+			Match LinkMatch = Regex.Match(Diagnostic.Message, ": reference to undefined symbol [`']([^`']+)");
+			if (LinkMatch.Success)
+			{
+				SymbolMatches.Add(LinkMatch.Groups[1].Value);
+			}
+
+			// Microsoft linker error:
+			//   Foo.cpp.obj : error LNK2001: unresolved external symbol \"private: virtual void __cdecl UAssetManager::InitializeAssetBundlesFromMetadata_Recursive(class UStruct const *,void const *,struct FAssetBundleData &,class FName,class TSet<void const *,struct DefaultKeyFuncs<void const *,0>,class FDefaultSetAllocator> &)const \" (?InitializeAssetBundlesFromMetadata_Recursive@UAssetManager@@EEBAXPEBVUStruct@@PEBXAEAUFAssetBundleData@@VFName@@AEAV?$TSet@PEBXU?$DefaultKeyFuncs@PEBX$0A@@@VFDefaultSetAllocator@@@@@Z)",
+			Match MicrosoftMatch = Regex.Match(Diagnostic.Message, ": unresolved external symbol \"([^\"]*)\"");
+			if(MicrosoftMatch.Success)
+			{
+				SymbolMatches.Add(MicrosoftMatch.Groups[1].Value);
+			}
+
+			// Clean up all the symbol names
+			SortedSet<string> SymbolNames = new SortedSet<string>(StringComparer.Ordinal);
+			foreach(string SymbolMatch in SymbolMatches)
+			{
+				string SymbolName = SymbolMatch;
+
+				// Remove any __declspec qualifiers
+				SymbolName = Regex.Replace(SymbolName, "(?<![^a-zA-Z_])__declspec\\([^\\)]+\\)", "");
+
+				// Remove any argument lists for functions (anything after the first paren)
+				SymbolName = Regex.Replace(SymbolName, "\\(.*$", "");
+
+				// Remove any decorators and type information (greedy match up to the last space)
+				SymbolName = Regex.Replace(SymbolName, "^.* ", "");
+
+				// Add it to the list
+				SymbolNames.Add(SymbolName);
+			}
+
+			// If we found any symbol names, create a fingerprint for them
+			if (SymbolNames.Count > 0)
+			{
+				Issue Issue = new Issue(Job.Project, Category, Job.Url, new IssueDiagnostic(JobStep.Name, JobStep.Url, Diagnostic.Message, Diagnostic.Url));
+				Issue.Identifiers.UnionWith(SymbolNames);
+				Issues.Add(Issue);
+				return true;
+			}
+
+			// Otherwise pass
+			return false;
+		}
+
+		public override List<ChangeInfo> FindCausers(PerforceConnection Perforce, Issue Issue, IReadOnlyList<ChangeInfo> Changes)
+		{
+			SortedSet<string> FileNamesWithoutPath = new SortedSet<string>(StringComparer.Ordinal);
+			foreach(string Identifier in Issue.Identifiers)
+			{
+				Match Match = Regex.Match(Identifier, "^[AUFS]([A-Z][A-Za-z_]*)::");
+				if(Match.Success)
+				{
+					FileNamesWithoutPath.Add(Match.Groups[1].Value + ".h");
+					FileNamesWithoutPath.Add(Match.Groups[1].Value + ".cpp");
+				}
+			}
+
+			if(FileNamesWithoutPath.Count > 0)
+			{
+				List<ChangeInfo> Causers = new List<ChangeInfo>();
+				foreach (ChangeInfo Change in Changes)
+				{
+					DescribeRecord DescribeRecord = GetDescribeRecord(Perforce, Change);
+					if (ContainsFileNames(DescribeRecord, FileNamesWithoutPath))
+					{
+						Causers.Add(Change);
+					}
+				}
+
+				if (Causers.Count > 0)
+				{
+					return Causers;
+				}
+			}
+
+			return base.FindCausers(Perforce, Issue, Changes);
+		}
+
+		public override string GetSummary(Issue Issue)
+		{
+			if(Issue.Identifiers.Count == 1)
+			{
+				return String.Format("Undefined symbol '{0}'", Issue.Identifiers.First());
+			}
+			else
+			{
+				return String.Format("Undefined symbols - '{0}'", String.Join("', '", Issue.Identifiers));
+			}
+		}
+	}
+}
diff --git a/Engine/Source/Programs/BuildAgent/Issues/PersistentState.cs b/Engine/Source/Programs/BuildAgent/Issues/PersistentState.cs
new file mode 100644
--- /dev/null
+++ b/Engine/Source/Programs/BuildAgent/Issues/PersistentState.cs
@@ -0,0 +1,106 @@
+// Copyright 1998-2019 Epic Games, Inc. All Rights Reserved.
+
+using System;
+using System.Collections.Generic;
+using System.Linq;
+using System.Runtime.Serialization;
+using System.Text;
+using System.Threading.Tasks;
+
+namespace BuildAgent.Issues
+{
+	/// <summary>
+	/// State tracked by the program between runs.
+	/// </summary>
+	[DataContract]
+	class PersistentState
+	{
+		/// <summary>
+		/// All open issues.
+		/// </summary>
+		[DataMember]
+		public List<Issue> Issues = new List<Issue>();
+
+		/// <summary>
+		/// Map of stream name to list of builds within it. Only retains a few builds in order to determine the last succesful CL for a build.
+		/// </summary>
+		[DataMember]
+		public Dictionary<string, List<IssueBuild>> Streams = new Dictionary<string, List<IssueBuild>>();
+
+		/// <summary>
+		/// Adds a new build for tracking
+		/// </summary>
+		/// <param name="Stream">The stream containing the build</param>
+		/// <param name="Build">The build to add</param>
+		public void AddBuild(string Stream, IssueBuild Build)
+		{
+			// Get the list of tracked builds for this stream
+			List<IssueBuild> Builds;
+			if (!Streams.TryGetValue(Stream, out Builds))
+			{
+				Builds = new List<IssueBuild>();
+				Streams.Add(Stream, Builds);
+			}
+
+			// Add this build to the tracked state data
+			int BuildIdx = Builds.BinarySearch(Build);
+			if (BuildIdx < 0)
+			{
+				BuildIdx = ~BuildIdx;
+				Builds.Insert(BuildIdx, Build);
+			}
+		}
+
+		/// <summary>
+		/// Finds the last build before the one given which executes the same steps. Used to determine the last succesful build before a failure.
+		/// </summary>
+		/// <param name="Stream">The stream to search for</param>
+		/// <param name="Change">The change to search for</param>
+		/// <param name="StepNames">The step names which the job step needs to have</param>
+		/// <returns>The last build known before the given changelist</returns>
+		public IssueBuild FindBuildBefore(string Stream, int Change, string StepName)
+		{
+			IssueBuild Result = null;
+
+			List<IssueBuild> Builds;
+			if(Streams.TryGetValue(Stream, out Builds))
+			{
+				for(int Idx = 0; Idx < Builds.Count && Builds[Idx].Change < Change; Idx++)
+				{
+					if(Builds[Idx].JobStepName == StepName)
+					{
+						Result = Builds[Idx];
+					}
+				}
+			}
+
+			return Result;
+		}
+
+		/// <summary>
+		/// Finds the first build after the one given which executes the same steps. Used to determine the last succesful build before a failure.
+		/// </summary>
+		/// <param name="Stream">The stream to search for</param>
+		/// <param name="Change">The change to search for</param>
+		/// <param name="StepName">Name of the job step to find</param>
+		/// <returns>The last build known before the given changelist</returns>
+		public IssueBuild FindBuildAfter(string Stream, int Change, string StepName)
+		{
+			IssueBuild Result = null;
+
+			List<IssueBuild> Builds;
+			if (Streams.TryGetValue(Stream, out Builds))
+			{
+				for (int Idx = Builds.Count - 1; Idx >= 0 && Builds[Idx].Change > Change; Idx--)
+				{
+					if (Builds[Idx].JobStepName == StepName)
+					{
+						Result = Builds[Idx];
+					}
+				}
+			}
+
+			return Result;
+		}
+	}
+}
diff --git a/Engine/Source/Programs/BuildAgent/Issues/UpdateIssuesMode.cs b/Engine/Source/Programs/BuildAgent/Issues/UpdateIssuesMode.cs
new file mode 100644
--- /dev/null
+++ b/Engine/Source/Programs/BuildAgent/Issues/UpdateIssuesMode.cs
@@ -0,0 +1,878 @@
+// Copyright 1998-2019 Epic Games, Inc. All Rights Reserved.
+
+using System;
+using System.Collections.Generic;
+using System.ComponentModel;
+using System.Diagnostics;
+using System.IO;
+using System.Linq;
+using System.Net;
+using System.Runtime.Serialization.Json;
+using System.Text;
+using System.Text.RegularExpressions;
+using System.Threading.Tasks;
+using System.Web;
+using System.Web.Script.Serialization;
+using Tools.DotNETCommon;
+using Tools.DotNETCommon.Perforce;
+using BuildAgent.Issues.Matchers;
+using BuildAgent.WebApi;
+
+namespace BuildAgent.Issues
+{
+	[ProgramMode("UpdateIssues", "Updates the UGS build health system with output from completed builds")]
+	class UpdateIssuesMode : ProgramMode
+	{
+		// Register all the pattern matchers
+		static readonly List<Matcher> Matchers = new List<Matcher>()
+		{ 
+			new CompileIssueMatcher(),
+			new UndefinedSymbolMatcher(),
+			new CopyrightNoticeMatcher(),
+			new ContentIssueMatcher()
+		};
+
+		static readonly Dictionary<string, Matcher> CategoryToMatcher = Matchers.ToDictionary(x => x.Category, x => x);
+
+		class CachedChangeInfo
+		{
+			public string Stream;
+			public int PrevChange;
+			public int NextChange;
+			public IReadOnlyList<ChangeInfo> Changes;
+		}
+
+		List<CachedChangeInfo> CachedChanges = new List<CachedChangeInfo>();
+
+		[CommandLine("-Server=")]
+		[Description("Url of the UGS metadata service")]
+		string ServerUrl = null;
+
+		[CommandLine("-InputFile=")]
+		[Description("Path to an input file containing completed jobs")]
+		FileReference InputFile = null;
+
+		[CommandLine("-StateFile=", Required = true)]
+		[Description("Path to a file containing persistent state")]
+		FileReference StateFile = null;
+
+		[CommandLine("-P4Port=")]
+		[Description("Server and port for P4 commands")]
+		string PerforcePort = null;
+
+		[CommandLine("-P4User=")]
+		[Description("Username to use for P4 commands")]
+		string PerforceUser = null;
+
+		[CommandLine("-Clean")]
+		[Description("Removes the existing state file and creates a new one")]
+		bool bClean = false;
+
+		[CommandLine("-KeepHistory")]
+		[Description("Retain history of all completed jobs in the state file, rather than purging after a few days.")]
+		bool bKeepHistory = false;
+
+		[CommandLine("-ReadOnly")]
+		[Description("Do not modify any files on disk.")]
+		bool bReadOnly = false;
+
+		[CommandLine("-SaveUnmatched")]
+		[Description("Path to a directory to save any unmatched issues")]
+		DirectoryReference SaveUnmatchedDir = null;
+
+		/// <summary>
+		/// Main command entry point
+		/// </summary>
+		/// <param name="Arguments">The command line arguments</param>
+		public override void Execute()
+		{
+			// Build a mapping from category to matching
+			Dictionary<string, Matcher> CategoryNameToMatcher = new Dictionary<string, Matcher>();
+			foreach (Matcher Matcher in Matchers)
+			{
+				CategoryNameToMatcher[Matcher.Category] = Matcher;
+			}
+
+			// Complete any interrupted operation to update the state file
+			CompleteStateTransaction(StateFile);
+
+			// Read the persistent data file
+			PersistentState State;
+			if (!bClean && FileReference.Exists(StateFile))
+			{
+				Log.TraceInformation("Reading persistent data from {0}", StateFile);
+				State = DeserializeJson<PersistentState>(StateFile);
+			}
+			else
+			{
+				Log.TraceInformation("Creating new persistent data");
+				State = new PersistentState();
+			}
+
+			// Fixup any issues loaded from disk
+			foreach(Issue Issue in State.Issues)
+			{
+				if(Issue.References == null)
+				{
+					Issue.References = new SortedSet<string>();
+				}
+			}
+
+			// Create the Perforce connection
+			PerforceConnection Perforce = new PerforceConnection(PerforcePort, PerforceUser, null);
+
+			// Process the input data
+			if(InputFile != null)
+			{
+				// Parse the input file
+				Log.TraceInformation("Reading build results from {0}", InputFile);
+				InputData InputData = DeserializeJson<InputData>(InputFile);
+
+				// Parse all the builds and add them to the persistent data
+				List<InputJob> InputJobs = InputData.Jobs.OrderBy(x => x.Change).ThenBy(x => x.Stream).ToList();
+				Stopwatch Timer = Stopwatch.StartNew();
+				foreach (InputJob InputJob in InputJobs)
+				{
+					// Add a new build for each job step
+					foreach(InputJobStep InputJobStep in InputJob.Steps)
+					{
+						IssueBuild NewBuild = new IssueBuild(InputJob.Change, InputJob.Name, InputJob.Url, InputJobStep.Name, InputJobStep.Url, null);
+						State.AddBuild(InputJob.Stream, NewBuild);
+					}
+
+					// Add all the job steps
+					List<InputJobStep> InputJobSteps = InputJob.Steps.OrderBy(x => x.Name).ToList();
+					foreach (InputJobStep InputJobStep in InputJobSteps)
+					{
+						if (InputJobStep.Diagnostics != null && InputJobStep.Diagnostics.Count > 0)
+						{
+							AddStep(Perforce, State, InputJob, InputJobStep);
+						}
+					}
+					
+					// Remove any steps which are empty
+					InputJob.Steps.RemoveAll(x => x.Diagnostics == null || x.Diagnostics.Count == 0);
+				}
+				InputJobs.RemoveAll(x => x.Steps.Count == 0);
+				Log.TraceInformation("Added jobs in {0}s", Timer.Elapsed.TotalSeconds);
+
+				// If there are any unmatched issues, save out the current state and remaining input
+				if(SaveUnmatchedDir != null && InputJobs.Count > 0)
+				{
+					DirectoryReference.CreateDirectory(SaveUnmatchedDir);
+					if(FileReference.Exists(StateFile))
+					{
+						FileReference.Copy(StateFile, FileReference.Combine(SaveUnmatchedDir, "State.json"), true);
+					}
+					SerializeJson(FileReference.Combine(SaveUnmatchedDir, "Input.json"), InputData);
+				}
+
+				// Try to find the next successful build for each stream, so we can close it as part of updating the server
+				for (int Idx = 0; Idx < State.Issues.Count; Idx++)
+				{
+					Issue Issue = State.Issues[Idx];
+					foreach(string Stream in Issue.Streams.Keys)
+					{
+						Dictionary<string, IssueHistory> StepNameToHistory = Issue.Streams[Stream];
+						foreach(string StepName in StepNameToHistory.Keys)
+						{
+							IssueHistory IssueHistory = StepNameToHistory[StepName];
+							if(IssueHistory.FailedBuilds.Count > 0 && IssueHistory.NextSuccessfulBuild == null)
+							{
+								// Find the successful build after this change
+								IssueBuild LastFailedBuild = IssueHistory.FailedBuilds[IssueHistory.FailedBuilds.Count - 1];
+								IssueHistory.NextSuccessfulBuild = State.FindBuildAfter(Stream, LastFailedBuild.Change, StepName);
+							}
+						}
+					}
+				}
+
+				// Find the change two days before the latest change being added
+				if(InputData.Jobs.Count > 0 && !bKeepHistory)
+				{
+					// Find all the unique change numbers for each stream
+					SortedSet<int> ChangeNumbers = new SortedSet<int>();
+					foreach (List<IssueBuild> Builds in State.Streams.Values)
+					{
+						ChangeNumbers.UnionWith(Builds.Select(x => x.Change));
+					}
+
+					// Get the latest change record
+					int LatestChangeNumber = InputData.Jobs.Min(x => x.Change);
+					ChangeRecord LatestChangeRecord = Perforce.GetChange(GetChangeOptions.None, LatestChangeNumber).Data;
+
+					// Step forward through all the changelists until we get to one we don't want to delete
+					int DeleteChangeNumber = -1;
+					foreach(int ChangeNumber in ChangeNumbers)
+					{
+						ChangeRecord ChangeRecord = Perforce.GetChange(GetChangeOptions.None, ChangeNumber).Data;
+						if (ChangeRecord.Date > LatestChangeRecord.Date - TimeSpan.FromDays(2))
+						{
+							break;
+						}
+						DeleteChangeNumber = ChangeNumber;
+					}
+
+					// Remove any builds we no longer want to track
+					foreach (List<IssueBuild> Builds in State.Streams.Values)
+					{
+						Builds.RemoveAll(x => x.Change <= DeleteChangeNumber);
+					}
+				}
+			}
+
+			// Mark any issues as resolved
+			foreach(Issue Issue in State.Issues)
+			{
+				if(Issue.IsResolved())
+				{
+					if(!Issue.ResolvedAt.HasValue)
+					{
+						Issue.ResolvedAt = DateTime.UtcNow;
+					}
+				}
+				else
+				{
+					if(Issue.ResolvedAt.HasValue)
+					{
+						Issue.ResolvedAt = null;
+					}
+				}
+			}
+
+			// If we're in read-only mode, don't write anything out
+			if(bReadOnly)
+			{
+				return;
+			}
+
+			// Save the persistent data
+			Log.TraceInformation("Writing persistent data to {0}", StateFile);
+			DirectoryReference.CreateDirectory(StateFile.Directory);
+			WriteState(StateFile, State);
+
+			// Synchronize with the server
+			if (ServerUrl != null)
+			{
+				// Post any issue updates
+				foreach(Issue Issue in State.Issues)
+				{
+					Matcher Matcher;
+					if(!CategoryNameToMatcher.TryGetValue(Issue.Category, out Matcher))
+					{
+						continue;
+					}
+
+					string Summary = Matcher.GetSummary(Issue);
+					if (Issue.Id == -1)
+					{
+						Log.TraceInformation("Adding issue: {0}", Issue);
+
+						if(Issue.PendingWatchers.Count == 0)
+						{
+							Log.TraceWarning("(No possible causers)");
+						}
+
+						ApiTypes.AddIssue IssueBody = new ApiTypes.AddIssue();
+						IssueBody.Project = Issue.Project;
+						IssueBody.Summary = Summary;
+
+						if(Issue.PendingWatchers.Count == 1)
+						{
+							IssueBody.Owner = Issue.PendingWatchers.First();
+						}
+
+						using(HttpWebResponse Response = SendHttpRequest(String.Format("{0}/api/issues", ServerUrl), "POST", IssueBody))
+						{
+							int ResponseCode = (int)Response.StatusCode;
+							if (!(ResponseCode >= 200 && ResponseCode <= 299))
+							{
+								throw new Exception("Unable to add issue");
+							}
+							Issue.Id = ParseHttpResponse<ApiTypes.AddIssueResponse>(Response).Id;
+						}
+
+						Issue.PostedSummary = Summary;
+						WriteState(StateFile, State);
+					}
+					else if(Issue.PostedSummary == null || !String.Equals(Issue.PostedSummary, Summary, StringComparison.Ordinal))
+					{
+						Log.TraceInformation("Updating issue {0}", Issue.Id);
+
+						ApiTypes.UpdateIssue IssueBody = new ApiTypes.UpdateIssue();
+						IssueBody.Summary = Summary;
+
+						using (HttpWebResponse Response = SendHttpRequest(String.Format("{0}/api/issues/{1}", ServerUrl, Issue.Id), "PUT", IssueBody))
+						{
+							int ResponseCode = (int)Response.StatusCode;
+							if (!(ResponseCode >= 200 && ResponseCode <= 299))
+							{
+								throw new Exception("Unable to add issue");
+							}
+						}
+
+						Issue.PostedSummary = Summary;
+						WriteState(StateFile, State);
+					}
+				}
+
+				// Add any new builds associated with issues
+				Dictionary<string, long> JobStepUrlToId = new Dictionary<string, long>(StringComparer.Ordinal);
+				foreach (Issue Issue in State.Issues)
+				{
+					foreach(KeyValuePair<string, Dictionary<string, IssueHistory>> StreamPair in Issue.Streams)
+					{
+						foreach(IssueHistory StreamHistory in StreamPair.Value.Values)
+						{
+							foreach(IssueBuild Build in StreamHistory.Builds)
+							{
+								if(!Build.bPostedToServer)
+								{
+									Log.TraceInformation("Adding {0} to issue {1}", Build.JobStepUrl, Issue.Id);
+
+									ApiTypes.AddBuild AddBuild = new ApiTypes.AddBuild();
+									AddBuild.Stream = StreamPair.Key;
+									AddBuild.Change = Build.Change;
+									AddBuild.JobName = Build.JobName;
+									AddBuild.JobUrl = Build.JobUrl;
+									AddBuild.JobStepName = Build.JobStepName;
+									AddBuild.JobStepUrl = Build.JobStepUrl;
+									AddBuild.ErrorUrl = Build.ErrorUrl;
+									AddBuild.Outcome = (Build == StreamHistory.PrevSuccessfulBuild || Build == StreamHistory.NextSuccessfulBuild)? ApiTypes.Outcome.Success : ApiTypes.Outcome.Error;
+
+									using (HttpWebResponse Response = SendHttpRequest(String.Format("{0}/api/issues/{1}/builds", ServerUrl, Issue.Id), "POST", AddBuild))
+									{
+										int ResponseCode = (int)Response.StatusCode;
+										if (!(ResponseCode >= 200 && ResponseCode <= 299))
+										{
+											throw new Exception("Unable to add build");
+										}
+										Build.Id = ParseHttpResponse<ApiTypes.AddBuildResponse>(Response).Id;
+									}
+
+									Build.bPostedToServer = true;
+									WriteState(StateFile, State);
+								}
+								if(Build.Id != -1)
+								{
+									JobStepUrlToId[Build.JobStepUrl] = Build.Id;
+								}
+							}
+						}
+					}
+				}
+
+				// Add any new diagnostics
+				foreach(Issue Issue in State.Issues)
+				{
+					foreach(IssueDiagnostic Diagnostic in Issue.Diagnostics)
+					{
+						if(!Diagnostic.bPostedToServer)
+						{
+							string Summary = Diagnostic.Message;
+
+							const int MaxLength = 40;
+							if(Summary.Length > MaxLength)
+							{
+								Summary = Summary.Substring(0, MaxLength).TrimEnd();
+							}
+
+							Log.TraceInformation("Adding diagnostic '{0}' to issue {1}", Summary, Issue.Id);
+
+							ApiTypes.AddDiagnostic AddDiagnostic = new ApiTypes.AddDiagnostic();
+
+							long BuildId;
+							if(Diagnostic.JobStepUrl != null && JobStepUrlToId.TryGetValue(Diagnostic.JobStepUrl, out BuildId))
+							{
+								AddDiagnostic.BuildId = BuildId;
+							}
+							else
+							{
+								Console.WriteLine("ERROR");
+							}
+
+							AddDiagnostic.Message = Diagnostic.Message;
+							AddDiagnostic.Url = Diagnostic.ErrorUrl;
+
+							using (HttpWebResponse Response = SendHttpRequest(String.Format("{0}/api/issues/{1}/diagnostics", ServerUrl, Issue.Id), "POST", AddDiagnostic))
+							{
+								int ResponseCode = (int)Response.StatusCode;
+								if (!(ResponseCode >= 200 && ResponseCode <= 299))
+								{
+									throw new Exception("Unable to add build");
+								}
+							}
+
+							Diagnostic.bPostedToServer = true;
+							WriteState(StateFile, State);
+						}
+					}
+				}
+
+				// Close any issues which are complete
+				for (int Idx = 0; Idx < State.Issues.Count; Idx++)
+				{
+					Issue Issue = State.Issues[Idx];
+					if (Issue.ResolvedAt.HasValue != Issue.bPostedResolved)
+					{
+						Log.TraceInformation("Setting issue {0} resolved flag to {1}", Issue.Id, Issue.ResolvedAt.HasValue);
+
+						ApiTypes.UpdateIssue UpdateBody = new ApiTypes.UpdateIssue();
+						UpdateBody.Resolved = Issue.ResolvedAt.HasValue;
+
+						using(HttpWebResponse Response = SendHttpRequest(String.Format("{0}/api/issues/{1}", ServerUrl, Issue.Id), "PUT", UpdateBody))
+						{
+							int ResponseCode = (int)Response.StatusCode;
+							if (!(ResponseCode >= 200 && ResponseCode <= 299))
+							{
+								throw new Exception("Unable to delete issue");
+							}
+						}
+
+						Issue.bPostedResolved = Issue.ResolvedAt.HasValue;
+						WriteState(StateFile, State);
+					}
+				}
+
+				// Update watchers on any open builds
+				foreach(Issue Issue in State.Issues)
+				{
+					while (Issue.PendingWatchers.Count > 0)
+					{
+						ApiTypes.Watcher Watcher = new ApiTypes.Watcher();
+						Watcher.UserName = Issue.PendingWatchers.First();
+
+						using (HttpWebResponse Response = SendHttpRequest(String.Format("{0}/api/issues/{1}/watchers", ServerUrl, Issue.Id), "POST", Watcher))
+						{
+							int ResponseCode = (int)Response.StatusCode;
+							if (!(ResponseCode >= 200 && ResponseCode <= 299))
+							{
+								throw new Exception("Unable to add watcher");
+							}
+						}
+
+						Issue.PendingWatchers.Remove(Watcher.UserName);
+						Issue.Watchers.Add(Watcher.UserName);
+
+						WriteState(StateFile, State);
+					}
+				}
+			}
+
+			// Remove any issues which have been resolved for 24 hours. We have to keep information about issues that have been fixed for some time; we may be updating the same job 
+			// multiple times while other steps are running, and we don't want to keep opening new issues for it. Also, it can take time for changes to propagate between streams.
+			DateTime RemoveIssueTime = DateTime.UtcNow - TimeSpan.FromHours(24.0);
+			for(int Idx = 0; Idx < State.Issues.Count; Idx++)
+			{
+				Issue Issue = State.Issues[Idx];
+				if(Issue.ResolvedAt.HasValue && Issue.ResolvedAt.Value < RemoveIssueTime)
+				{
+					State.Issues.RemoveAt(Idx--);
+					WriteState(StateFile, State);
+					continue;
+				}
+			}
+
+			// TODO: VERIFY ISSUES ARE CLOSED
+		}
+
+		/// <summary>
+		/// Sends an arbitrary HTTP request
+		/// </summary>
+		/// <param name="Url">Endpoint to send to</param>
+		/// <param name="Method">The method to use for sending the request</param>
+		/// <param name="BodyObject">Object to be serialized as json in the body</param>
+		/// <returns>HTTP response</returns>
+		HttpWebResponse SendHttpRequest(string Url, string Method, object BodyObject)
+		{
+			// Create the request
+			HttpWebRequest Request = (HttpWebRequest)WebRequest.Create(Url);
+			Request.ContentType = "application/json";
+			Request.Method = Method;
+
+			string BodyText = null;
+			if (BodyObject != null)
+			{
+				BodyText = new JavaScriptSerializer().Serialize(BodyObject);
+				byte[] BodyData = Encoding.UTF8.GetBytes(BodyText);
+				using (Stream RequestStream = Request.GetRequestStream())
+				{
+					RequestStream.Write(BodyData, 0, BodyData.Length);
+				}
+			}
+
+			// Read the response
+			try
+			{
+				return (HttpWebResponse)Request.GetResponse();
+			}
+			catch(Exception Ex)
+			{
+				ExceptionUtils.AddContext(Ex, String.Format("Url: {0}", Url));
+				ExceptionUtils.AddContext(Ex, String.Format("Method: {0}", Method));
+				ExceptionUtils.AddContext(Ex, String.Format("Body: {0}", BodyText));
+				throw;
+			}
+		}
+
+		/// <summary>
+		/// Parses an HTTP response object as JSON
+		/// </summary>
+		/// <typeparam name="T">The type of object to parse</typeparam>
+		/// <param name="Response">The web response instance</param>
+		/// <returns>Response object</returns>
+		T ParseHttpResponse<T>(HttpWebResponse Response)
+		{
+			using (StreamReader ResponseReader = new StreamReader(Response.GetResponseStream(), Encoding.Default))
+			{
+				string ResponseContent = ResponseReader.ReadToEnd();
+				return new JavaScriptSerializer().Deserialize<T>(ResponseContent);
+			}
+		}
+
+		/// <summary>
+		/// Adds diagnostics from a job step into the issue database
+		/// </summary>
+		/// <param name="Perforce">Perforce connection used to find possible causers</param>
+		/// <param name="State">The current set of tracked issues</param>
+		/// <param name="Build">The new build</param>
+		/// <param name="PreviousChange">The last changelist that was built before this one</param>
+		/// <param name="InputJob">Job containing the step to add</param>
+		/// <param name="InputJobStep">The job step to add</param>
+		void AddStep(PerforceConnection Perforce, PersistentState State, InputJob InputJob, InputJobStep InputJobStep)
+		{
+			// Create a lazily evaluated list of changes that are responsible for any errors
+			Lazy<IReadOnlyList<ChangeInfo>> LazyChanges = new Lazy<IReadOnlyList<ChangeInfo>>(() => FindChanges(Perforce, State, InputJob));
+
+			// Create issues for any diagnostics in this step
+			List<Issue> InputIssues = new List<Issue>();
+			foreach(Matcher Matcher in Matchers)
+			{
+				Matcher.Match(InputJob, InputJobStep, InputJobStep.Diagnostics, InputIssues);
+			}
+
+			// Merge the issues together
+			List<Issue> NewIssues = new List<Issue>();
+			foreach(Issue InputIssue in InputIssues)
+			{
+				Issue OutputIssue = MergeIntoExistingIssue(Perforce, State, InputJob, InputJobStep, InputIssue, LazyChanges);
+				if(OutputIssue == null)
+				{
+					NewIssues.Add(InputIssue);
+					State.Issues.Add(InputIssue);
+					OutputIssue = InputIssue;
+				}
+				AddFailureToIssue(OutputIssue, InputJob, InputJobStep, InputIssue.Diagnostics[0].ErrorUrl, State);
+			}
+
+			// Update the watchers for any new issues
+			foreach(Issue NewIssue in NewIssues)
+			{
+				IReadOnlyList<ChangeInfo> Changes = LazyChanges.Value;
+				if (Changes != null)
+				{
+					// Find the pattern matcher for this issue
+					Matcher Matcher = CategoryToMatcher[NewIssue.Category];
+
+					// Update the causers
+					List<ChangeInfo> Causers = Matcher.FindCausers(Perforce, NewIssue, Changes);
+					foreach (ChangeInfo Causer in Causers)
+					{
+						NewIssue.SourceChanges.UnionWith(Causer.SourceChanges);
+						NewIssue.PendingWatchers.Add(Causer.Record.User);
+					}
+				}
+			}
+		}
+
+		/// <summary>
+		/// Finds or adds an issue for a particular issue
+		/// </summary>
+		/// <param name="Perforce">Perforce connection used to find possible causers</param>
+		/// <param name="State">The current set of tracked issues</param>
+		/// <param name="Build">The new build</param>
+		/// <param name="PreviousChange">The last changelist that was built before this one</param>
+		/// <param name="InputJob">Job containing the step to add</param>
+		/// <param name="InputJobStep">The job step to add</param>
+		Issue MergeIntoExistingIssue(PerforceConnection Perforce, PersistentState State, InputJob InputJob, InputJobStep InputJobStep, Issue InputIssue, Lazy<IReadOnlyList<ChangeInfo>> LazyChanges)
+		{
+			// Find the pattern matcher for this fingerprint
+			Matcher Matcher = CategoryToMatcher[InputIssue.Category];
+
+			// Check if it can be added to an existing open issue
+			foreach (Issue Issue in State.Issues)
+			{
+				// Check this issue already exists in the current stream
+				Dictionary<string, IssueHistory> StepNameToHistory;
+				if(!Issue.Streams.TryGetValue(InputJob.Stream, out StepNameToHistory))
+				{
+					continue;
+				}
+
+				// Check that this issue has not already been closed
+				IssueHistory History;
+				if (StepNameToHistory.TryGetValue(InputJobStep.Name, out History))
+				{
+					if(!History.CanAddFailedBuild(InputJob.Change))
+					{
+						continue;
+					}
+				}
+				else
+				{
+					if(!StepNameToHistory.Values.Any(x => x.CanAddFailedBuild(InputJob.Change)))
+					{
+						continue;
+					}
+				}
+
+				// Try to merge the fingerprint
+				if(!Matcher.CanMerge(InputIssue, Issue))
+				{
+					continue;
+				}
+
+				// Add the new build
+				Matcher.Merge(InputIssue, Issue);
+				return Issue;
+			}
+
+			// Check if this issue can be merged with an issue built in another stream
+			IReadOnlyList<ChangeInfo> Changes = LazyChanges.Value;
+			if(Changes != null && Changes.Count > 0)
+			{
+				SortedSet<int> SourceChanges = new SortedSet<int>(Changes.SelectMany(x => x.SourceChanges));
+				foreach (Issue Issue in State.Issues)
+				{
+					// Check if this issue does not already contain this stream, but contains one of the causing changes
+					if (Issue.Streams.ContainsKey(InputJob.Stream))
+					{
+						continue;
+					}
+					if(!SourceChanges.Any(x => Issue.SourceChanges.Contains(x)))
+					{
+						continue;
+					}
+					if (!Matcher.CanMerge(InputIssue, Issue))
+					{
+						continue;
+					}
+
+					// Merge the issue
+					Matcher.Merge(InputIssue, Issue);
+					return Issue;
+				}
+			}
+
+			// Check if it can be merged into an issue that's been created for this job. We only do this after exhausting all other options.
+			foreach (Issue Issue in State.Issues)
+			{
+				if(Issue.InitialJobUrl == InputIssue.InitialJobUrl && Matcher.CanMergeInitialJob(InputIssue, Issue))
+				{
+					Matcher.Merge(InputIssue, Issue);
+					return Issue;
+				}
+			}
+
+			return null;
+		}
+
+		/// <summary>
+		/// Creates a TrackedBuild instance for the given jobstep
+		/// </summary>
+		/// <param name="InputJob">The job to create a build for</param>
+		/// <param name="InputJobStep">The step to create a build for</param>
+		/// <param name="InputErrorUrl">The error Url</param>
+		/// <returns>New build instance</returns>
+		IssueBuild CreateBuildForJobStep(InputJob InputJob, InputJobStep InputJobStep, string InputErrorUrl)
+		{
+			return new IssueBuild(InputJob.Change, InputJob.Name, InputJob.Url, InputJobStep.Name, InputJobStep.Url, InputErrorUrl);
+		}
+
+		/// <summary>
+		/// Adds a new build history for a stream
+		/// </summary>
+		/// <param name="Issue">The issue to add a build to</param>
+		/// <param name="InputJob">The job containing the error</param>
+		/// <param name="InputJobStep">The job step containing the error</param>
+		/// <param name="InputErrorUrl">Url of the error</param>
+		/// <param name="State">Current persistent state. Used to find previous build history.</param>
+		void AddFailureToIssue(Issue Issue, InputJob InputJob, InputJobStep InputJobStep, string InputErrorUrl, PersistentState State)
+		{
+			// Find or add a step name to history mapping
+			Dictionary<string, IssueHistory> StepNameToHistory;
+			if(!Issue.Streams.TryGetValue(InputJob.Stream, out StepNameToHistory))
+			{
+				StepNameToHistory = new Dictionary<string, IssueHistory>();
+				Issue.Streams.Add(InputJob.Stream, StepNameToHistory);
+			}
+
+			// Find or add a history for this step
+			IssueHistory History;
+			if(!StepNameToHistory.TryGetValue(InputJobStep.Name, out History))
+			{
+				History = new IssueHistory(State.FindBuildBefore(InputJob.Stream, InputJob.Change, InputJobStep.Name));
+				StepNameToHistory.Add(InputJobStep.Name, History);
+			}
+
+			// Add the new build
+			History.AddFailedBuild(CreateBuildForJobStep(InputJob, InputJobStep, InputErrorUrl));
+		}
+
+		/// <summary>
+		/// Find all changes PLUS all robomerge source changes
+		/// </summary>
+		/// <param name="Perforce">The Perforce connection to use</param>
+		/// <param name="State">State of </param>
+		/// <param name="InputJob">The job that failed</param>
+		/// <returns>Set of changelist numbers</returns>
+		IReadOnlyList<ChangeInfo> FindChanges(PerforceConnection Perforce, PersistentState State, InputJob InputJob)
+		{
+			// List of changes since the last successful build in this stream
+			IReadOnlyList<ChangeInfo> Changes = null;
+
+			// Find the previous changelist that was built in this stream
+			List<IssueBuild> StreamBuilds;
+			if (State.Streams.TryGetValue(InputJob.Stream, out StreamBuilds))
+			{
+				// Find the last change submitted to this stream before it started failing
+				int LastChange = -1;
+				for (int Idx = 0; Idx < StreamBuilds.Count && StreamBuilds[Idx].Change < InputJob.Change; Idx++)
+				{
+					LastChange = StreamBuilds[Idx].Change;
+				}
+
+				// Allow adding to any open issue that contains changes merged from other branches
+				if (LastChange != -1)
+				{
+					// Query for all the changes since then
+					Changes = FindChanges(Perforce, InputJob.Stream, LastChange, InputJob.Change);
+				}
+			}
+			return Changes;
+		}
+
+		/// <summary>
+		/// Find all changes PLUS all robomerge source changes
+		/// </summary>
+		/// <param name="Perforce">The Perforce connection to use</param>
+		/// <param name="Stream">The stream to query changes from</param>
+		/// <param name="PrevChange">The first change in the range to query</param>
+		/// <param name="NextChange">The last change in the range to query</param>
+		/// <returns>Set of changelist numbers</returns>
+		IReadOnlyList<ChangeInfo> FindChanges(PerforceConnection Perforce, string Stream, int PrevChange, int NextChange)
+		{
+			CachedChangeInfo CachedInfo = CachedChanges.FirstOrDefault(x => x.Stream == Stream && x.PrevChange == PrevChange && x.NextChange == NextChange);
+			if(CachedInfo == null)
+			{
+				// Query for all the changes since then
+				List<ChangesRecord> ChangeRecords = Perforce.Changes(ChangesOptions.LongOutput, null, -1, ChangeStatus.Submitted, null, String.Format("{0}/...@{1},{2}", Stream, PrevChange, NextChange)).Data.ToList();
+
+				// Figure out all the original changelists that these were merged from, and see if any of those matches with an existing issue
+				List<ChangeInfo> Changes = new List<ChangeInfo>();
+				foreach (ChangesRecord ChangeRecord in ChangeRecords)
+				{
+					ChangeInfo Change = new ChangeInfo();
+					Change.Record = ChangeRecord;
+					Change.SourceChanges.Add(ChangeRecord.Number);
+					Changes.Add(Change);
+
+					Match SourceMatch = Regex.Match(ChangeRecord.Description, "^#ROBOMERGE-SOURCE: (.*)$", RegexOptions.Multiline);
+					if (SourceMatch.Success)
+					{
+						string SourceText = SourceMatch.Groups[1].Value;
+						foreach (Match ChangeMatch in Regex.Matches(SourceText, "CL\\s*(\\d+)"))
+						{
+							int SourceChange;
+							if (int.TryParse(ChangeMatch.Groups[1].Value, out SourceChange))
+							{
+								Change.SourceChanges.Add(SourceChange);
+							}
+						}
+					}
+				}
+
+				// Create the new cached info
+				CachedInfo = new CachedChangeInfo() { Stream = Stream, PrevChange = PrevChange, NextChange = NextChange, Changes = Changes };
+				CachedChanges.Add(CachedInfo);
+			}
+			return CachedInfo.Changes;
+		}
+
+		/// <summary>
+		/// Gets the path to a temporary file used for ensuring that serialization from the state file is transactional
+		/// </summary>
+		/// <param name="StateFile">Path to the state file</param>
+		/// <returns>Path to the temporary state transaction file</returns>
+		static FileReference GetStateTransactionFile(FileReference StateFile)
+		{
+			return new FileReference(StateFile.FullName + ".transaction");
+		}
+
+		/// <summary>
+		/// Completes an interrupted transaction to write the state file
+		/// </summary>
+		/// <param name="StateFile">Path to the state file</param>
+		static void CompleteStateTransaction(FileReference StateFile)
+		{
+			if(!FileReference.Exists(StateFile))
+			{
+				FileReference StateTransactionFile = GetStateTransactionFile(StateFile);
+				if(FileReference.Exists(StateTransactionFile))
+				{
+					Log.TraceInformation("Completing partial transaction through {0}", StateTransactionFile);
+					FileReference.Move(StateTransactionFile, StateFile);
+				}
+			}
+		}
+
+		/// <summary>
+		/// Writes the state to disk in a way that can be recovered if the operation is interrupted
+		/// </summary>
+		/// <param name="StateFile">The file to write to</param>
+		/// <param name="State">The state object</param>
+		static void WriteState(FileReference StateFile, PersistentState State)
+		{
+			Stopwatch Timer = Stopwatch.StartNew();
+
+			// Write out the state to the transaction file
+			FileReference StateTransactionFile = GetStateTransactionFile(StateFile);
+			SerializeJson(StateTransactionFile, State);
+
+			// Remove the original file, then move the transaction file into place
+			FileReference.Delete(StateFile);
+			FileReference.Move(StateTransactionFile, StateFile);
+
+			Log.TraceInformation("Took {0}s to write state", Timer.Elapsed.TotalSeconds);
+		}
+
+		/// <summary>
+		/// Serializes an object to a file as JSON
+		/// </summary>
+		/// <param name="Location">Location of the file</param>
+		/// <param name="Object">Object to serialize</param>
+		static void SerializeJson(FileReference Location, object Object)
+		{
+			using (MemoryStream Stream = new MemoryStream())
+			{
+				DataContractJsonSerializer InputFileDataSerializer = new DataContractJsonSerializer(Object.GetType());
+				InputFileDataSerializer.WriteObject(Stream, Object);
+				FileReference.WriteAllBytes(Location, Stream.ToArray());
+			}
+		}
+
+		/// <summary>
+		/// Deserializes a file at the given location
+		/// </summary>
+		/// <param name="Location">Location of the file</param>
+		/// <returns>Deserialized object</returns>
+		static T DeserializeJson<T>(FileReference Location) where T : class
+		{
+			using (FileStream Stream = File.Open(Location.FullName, FileMode.Open, FileAccess.Read, FileShare.Read))
+			{
+				DataContractJsonSerializer InputFileDataSerializer = new DataContractJsonSerializer(typeof(T));
+				return (T)InputFileDataSerializer.ReadObject(Stream);
+			}
+		}
+	}
+}
