Disable unused member variable warnings.

#rnx

diff --git a/Engine/Source/Programs/BuildAgent/BuildAgent.csproj b/Engine/Source/Programs/BuildAgent/BuildAgent.csproj
--- a/Engine/Source/Programs/BuildAgent/BuildAgent.csproj
+++ b/Engine/Source/Programs/BuildAgent/BuildAgent.csproj
@@ -32,6 +32,7 @@
     <DefineConstants>TRACE</DefineConstants>
     <ErrorReport>prompt</ErrorReport>
     <WarningLevel>4</WarningLevel>
+    <NoWarn>CS0649</NoWarn>
   </PropertyGroup>
   <Import Project="NotForLicensees\EpicInternal.props" Condition="Exists('NotForLicensees\EpicInternal.props')" />
   <ItemGroup>
