Fixes for BuildAgent:

* Prevent TerminateAllChildProcesses from being called before listeners have been disposed. This was terminating ectool before it had finished updating properties.
* Propagate exit code from child process as the exit code for BuildAgent.exe.
* Write EC XML output to a temporary file and move it into place, to fix issues where a job was aborted mid-write.

#rnx

diff --git a/Engine/Source/Programs/BuildAgent/Issues/UpdateIssuesMode.cs b/Engine/Source/Programs/BuildAgent/Issues/UpdateIssuesMode.cs
--- a/Engine/Source/Programs/BuildAgent/Issues/UpdateIssuesMode.cs
+++ b/Engine/Source/Programs/BuildAgent/Issues/UpdateIssuesMode.cs
@@ -84,7 +84,7 @@ namespace BuildAgent.Issues
 		/// Main command entry point
 		/// </summary>
 		/// <param name="Arguments">The command line arguments</param>
-		public override void Execute()
+		public override int Execute()
 		{
 			// Build a mapping from category to matching
 			Dictionary<string, Matcher> CategoryNameToMatcher = new Dictionary<string, Matcher>();
@@ -243,7 +243,7 @@ namespace BuildAgent.Issues
 			// If we're in read-only mode, don't write anything out
 			if(bReadOnly)
 			{
-				return;
+				return 0;
 			}
 
 			// Save the persistent data
@@ -474,6 +474,7 @@ namespace BuildAgent.Issues
 			}
 
 			// TODO: VERIFY ISSUES ARE CLOSED
+			return 0;
 		}
 
 		/// <summary>
diff --git a/Engine/Source/Programs/BuildAgent/Program.cs b/Engine/Source/Programs/BuildAgent/Program.cs
--- a/Engine/Source/Programs/BuildAgent/Program.cs
+++ b/Engine/Source/Programs/BuildAgent/Program.cs
@@ -133,8 +133,7 @@ namespace BuildAgent
 			}
 
 			// Execute all the commands
-			Mode.Execute();
-			return 0;
+			return Mode.Execute();
 		}
 
 		/// <summary>
diff --git a/Engine/Source/Programs/BuildAgent/ProgramMode.cs b/Engine/Source/Programs/BuildAgent/ProgramMode.cs
--- a/Engine/Source/Programs/BuildAgent/ProgramMode.cs
+++ b/Engine/Source/Programs/BuildAgent/ProgramMode.cs
@@ -22,6 +22,6 @@ namespace BuildAgent
 			return CommandLineArguments.GetParameters(GetType());
 		}
 
-		public abstract void Execute();
+		public abstract int Execute();
 	}
 }
diff --git a/Engine/Source/Programs/BuildAgent/Run/Listeners/ElectricCommanderListener.cs b/Engine/Source/Programs/BuildAgent/Run/Listeners/ElectricCommanderListener.cs
--- a/Engine/Source/Programs/BuildAgent/Run/Listeners/ElectricCommanderListener.cs
+++ b/Engine/Source/Programs/BuildAgent/Run/Listeners/ElectricCommanderListener.cs
@@ -151,7 +151,9 @@ namespace BuildAgent.Listeners
 				}
 
 				// Write them to disk
-				Write(OutputFile, Errors);
+				FileReference TempOutputFile = new FileReference(OutputFile.FullName + ".new");
+				Write(TempOutputFile, Errors);
+				FileUtils.ForceMoveFile(TempOutputFile, OutputFile);
 
 				// On the first run, set the path to the diagnostics file
 				if (!bHasSetDiagFile)
diff --git a/Engine/Source/Programs/BuildAgent/Run/RunMode.cs b/Engine/Source/Programs/BuildAgent/Run/RunMode.cs
--- a/Engine/Source/Programs/BuildAgent/Run/RunMode.cs
+++ b/Engine/Source/Programs/BuildAgent/Run/RunMode.cs
@@ -70,8 +70,10 @@ namespace BuildAgent.Run
 			ProgramArguments = Arguments.GetPositionalArguments();
 		}
 
-		public override void Execute()
+		public override int Execute()
 		{
+			int ExitCode = 0;
+
 			// Auto-register all the known matchers in this assembly
 			List<IErrorMatcher> Matchers = new List<IErrorMatcher>();
 			foreach (Type Type in Assembly.GetExecutingAssembly().GetTypes())
@@ -151,9 +153,8 @@ namespace BuildAgent.Run
 					{
 						Func<string> ReadLine = new LineFilter(() => ReadProcessLine(Process, CancellationToken)).ReadLine;
 						ProcessErrors(ReadLine, Matchers, IgnorePatterns, Listeners, bNoWarnings);
+						ExitCode = Process.ExitCode;
 					}
-
-					ProcessUtils.TerminateChildProcesses();
 				}
 			}
 			finally
@@ -163,6 +164,10 @@ namespace BuildAgent.Run
 					Listener.Dispose();
 				}
 			}
+
+			// Kill off any remaining child processes
+			ProcessUtils.TerminateChildProcesses();
+			return ExitCode;
 		}
 
 		/// <summary>
diff --git a/Engine/Source/Programs/BuildAgent/Server/ServerMode.cs b/Engine/Source/Programs/BuildAgent/Server/ServerMode.cs
--- a/Engine/Source/Programs/BuildAgent/Server/ServerMode.cs
+++ b/Engine/Source/Programs/BuildAgent/Server/ServerMode.cs
@@ -182,7 +182,7 @@ namespace BuildAgent.WebApi
 			return Parameters;
 		}
 
-		public override void Execute()
+		public override int Execute()
 		{
 			// Create the request
 			HttpWebRequest Request = (HttpWebRequest)WebRequest.Create(CommandUrl.ToString());
@@ -208,6 +208,8 @@ namespace BuildAgent.WebApi
 					Console.WriteLine(Json.Format(ResponseContent));
 				}
 			}
+
+			return 0;
 		}
 
 		static bool TryGetCommandArgument(CommandLineArguments Arguments, out CommandInfo Command)
diff --git a/Engine/Source/Programs/BuildAgent/Workspace/WorkspaceMode.cs b/Engine/Source/Programs/BuildAgent/Workspace/WorkspaceMode.cs
--- a/Engine/Source/Programs/BuildAgent/Workspace/WorkspaceMode.cs
+++ b/Engine/Source/Programs/BuildAgent/Workspace/WorkspaceMode.cs
@@ -57,7 +57,7 @@ namespace BuildAgent.Workspace
 			}
 		}
 
-		public override void Execute()
+		public override int Execute()
 		{
 			// Create the repo
 			Repository Repo = CreateOrLoadRepository(ServerAndPort, UserName, BaseDir, bOverwrite);
@@ -73,6 +73,7 @@ namespace BuildAgent.Workspace
 			Trace.Listeners.Add(new TextWriterTraceListener(new StreamWriter(LogFile.FullName), "LogTraceListener"));
 
 			Execute(Repo);
+			return 0;
 		}
 
 		protected abstract void Execute(Repository Repo);
