BuildAgent: Minor fixes.

* Fix help messages when specifying a command.
* Add support for parsing long types.

#rnx

diff --git a/Engine/Source/Programs/BuildAgent/Server/ServerMode.cs b/Engine/Source/Programs/BuildAgent/Server/ServerMode.cs
--- a/Engine/Source/Programs/BuildAgent/Server/ServerMode.cs
+++ b/Engine/Source/Programs/BuildAgent/Server/ServerMode.cs
@@ -37,10 +37,10 @@ namespace BuildAgent.WebApi
 		{
 			// Issues
 			new CommandInfo("AddIssue", "POST", "/api/issues", typeof(ApiTypes.AddIssue)),
-			new CommandInfo("GetIssue", "GET", "/api/issues/{id}"),
+			new CommandInfo("GetIssue", "GET", "/api/issues/{Issue}"),
 			new CommandInfo("GetIssues", "GET", "/api/issues", OptionalParams: new string[] { "includeresolved", "maxresults", "user" }),
-			new CommandInfo("UpdateIssue", "PUT", "/api/issues/{Id}", typeof(ApiTypes.UpdateIssue)),
-			new CommandInfo("DeleteIssue", "DELETE", "/api/issues/{Id}"),
+			new CommandInfo("UpdateIssue", "PUT", "/api/issues/{Issue}", typeof(ApiTypes.UpdateIssue)),
+			new CommandInfo("DeleteIssue", "DELETE", "/api/issues/{Issue}"),
 
 			// Builds
 			new CommandInfo("AddBuild", "POST", "/api/issues/{Issue}/builds", typeof(ApiTypes.AddBuild)),
@@ -144,7 +144,7 @@ namespace BuildAgent.WebApi
 		{
 			List<KeyValuePair<string, string>> Parameters = new List<KeyValuePair<string, string>>();
 
-			if(TryGetCommandArgument(Arguments, out Command))
+			if(Command != null)
 			{
 				Parameters.Insert(0, new KeyValuePair<string, string>(Command.Name, "Name of the command to execute"));
 
@@ -276,6 +276,15 @@ namespace BuildAgent.WebApi
 						}
 						Field.SetValue(Instance, IntValue);
 					}
+					else if (FieldType == typeof(long))
+					{
+						long LongValue;
+						if (!long.TryParse(Value, out LongValue))
+						{
+							throw new CommandLineArgumentException(String.Format("Invalid value for '{0}' - {1} is not an integer.", Field.Name, Value));
+						}
+						Field.SetValue(Instance, LongValue);
+					}
 					else if (FieldType == typeof(bool))
 					{
 						bool BoolValue;
