BuildAgent: Minor improvements to error output.

#rnx

diff --git a/Engine/Source/Programs/BuildAgent/Program.cs b/Engine/Source/Programs/BuildAgent/Program.cs
--- a/Engine/Source/Programs/BuildAgent/Program.cs
+++ b/Engine/Source/Programs/BuildAgent/Program.cs
@@ -82,19 +82,22 @@ namespace BuildAgent
 			}
 
 			// Get the command name
-			string CommandName = Args[ModeIndex];
+			string ModeName = Args[ModeIndex];
 
 			// Get the command type
-			Type CommandType;
-			if (!ModeNameToType.TryGetValue(CommandName, out CommandType))
+			Type ModeType;
+			if (!ModeNameToType.TryGetValue(ModeName, out ModeType))
 			{
-				Log.TraceError("Unknown command '{0}'.", CommandName);
+				Log.TraceError("Unknown command '{0}'.", ModeName);
 				Log.TraceInformation("");
 				Log.TraceInformation("Available commands");
 				PrintCommands(ModeNameToType);
 				return 1;
 			}
 
+			// Update the mode name to use the correct casing, in case we need it for error messages
+			ModeName = ModeType.GetCustomAttribute<ProgramModeAttribute>().Name;
+
 			// Build an argument list for the command, including all the global arguments as well as arguments until the next command
 			List<string> ArgumentList = new List<string>();
 			for (int Idx = 0; Idx < Args.Length; Idx++)
@@ -107,12 +110,12 @@ namespace BuildAgent
 			CommandLineArguments ModeArguments = new CommandLineArguments(ArgumentList.ToArray());
 
 			// Create the command instance
-			ProgramMode Mode = (ProgramMode)Activator.CreateInstance(CommandType);
+			ProgramMode Mode = (ProgramMode)Activator.CreateInstance(ModeType);
 
 			// If the help flag is specified, print the help info and exit immediately
 			if (ModeArguments.HasOption("-Help"))
 			{
-				HelpUtils.PrintHelp(CommandName, HelpUtils.GetDescription(CommandType), Mode.GetParameters(ModeArguments));
+				HelpUtils.PrintHelp(ModeName, HelpUtils.GetDescription(ModeType), Mode.GetParameters(ModeArguments));
 				return 1;
 			}
 
@@ -124,9 +127,9 @@ namespace BuildAgent
 			}
 			catch (CommandLineArgumentException Ex)
 			{
-				Log.TraceError("{0}: {1}", CommandName, Ex.Message);
+				Log.TraceError("{0}: {1}", ModeName, Ex.Message);
 				Log.TraceInformation("");
-				Log.TraceInformation("Arguments for {0}:", CommandName);
+				Log.TraceInformation("Arguments for {0}:", ModeName);
 
 				HelpUtils.PrintTable(Mode.GetParameters(ModeArguments), 4, 24);
 				return 1;
diff --git a/Engine/Source/Programs/BuildAgent/Server/ServerMode.cs b/Engine/Source/Programs/BuildAgent/Server/ServerMode.cs
--- a/Engine/Source/Programs/BuildAgent/Server/ServerMode.cs
+++ b/Engine/Source/Programs/BuildAgent/Server/ServerMode.cs
@@ -86,10 +86,11 @@ namespace BuildAgent.WebApi
 			// Parse the server argument
 			if (String.IsNullOrEmpty(Server))
 			{
-				Server = Environment.GetEnvironmentVariable("UGS_METADATA_SERVER_URL");
+				const string ServerEnvVarName = "METADATA_SERVER_URL";
+				Server = Environment.GetEnvironmentVariable(ServerEnvVarName);
 				if (String.IsNullOrEmpty(Server))
 				{
-					throw new CommandLineArgumentException("Missing -Server=... argument.");
+					throw new CommandLineArgumentException(String.Format("Missing -Server=... argument or {0} environment variable.", ServerEnvVarName));
 				}
 			}
 
