BuildAgent: Tweak help messages.

#rnx

diff --git a/Engine/Source/Programs/BuildAgent/Program.cs b/Engine/Source/Programs/BuildAgent/Program.cs
--- a/Engine/Source/Programs/BuildAgent/Program.cs
+++ b/Engine/Source/Programs/BuildAgent/Program.cs
@@ -78,6 +78,9 @@ namespace BuildAgent
 				Log.TraceInformation("Commands:");
 
 				PrintCommands(ModeNameToType);
+
+				Log.TraceInformation("");
+				Log.TraceInformation("Specify \"Command -Help\" for command-specific help");
 				return 0;
 			}
 
