BuildAgent: Add a listener which outputs input data for updating build issues.

#rns

diff --git a/Engine/Source/Programs/BuildAgent/BuildAgent.csproj b/Engine/Source/Programs/BuildAgent/BuildAgent.csproj
--- a/Engine/Source/Programs/BuildAgent/BuildAgent.csproj
+++ b/Engine/Source/Programs/BuildAgent/BuildAgent.csproj
@@ -77,6 +77,7 @@
     <Compile Include="Run\LineFilter.cs" />
     <Compile Include="Run\Listeners\ElectricCommanderListener.cs" />
     <Compile Include="Run\Listeners\DebugOutputListener.cs" />
+    <Compile Include="Run\Listeners\IssuesListener.cs" />
     <Compile Include="Run\Matchers\**\*.cs" />
     <Compile Include="Program.cs" />
     <Compile Include="Run\ErrorMatch.cs" />
diff --git a/Engine/Source/Programs/BuildAgent/Run/Listeners/DebugOutputListener.cs b/Engine/Source/Programs/BuildAgent/Run/Listeners/DebugOutputListener.cs
--- a/Engine/Source/Programs/BuildAgent/Run/Listeners/DebugOutputListener.cs
+++ b/Engine/Source/Programs/BuildAgent/Run/Listeners/DebugOutputListener.cs
@@ -8,7 +8,7 @@ using System.Linq;
 using System.Text;
 using System.Threading.Tasks;
 
-namespace BuildAgent.Listeners
+namespace BuildAgent.Run.Listeners
 {
 	/// <summary>
 	/// Debug error listener. Just prints the error output to the log.
diff --git a/Engine/Source/Programs/BuildAgent/Run/Listeners/ElectricCommanderListener.cs b/Engine/Source/Programs/BuildAgent/Run/Listeners/ElectricCommanderListener.cs
--- a/Engine/Source/Programs/BuildAgent/Run/Listeners/ElectricCommanderListener.cs
+++ b/Engine/Source/Programs/BuildAgent/Run/Listeners/ElectricCommanderListener.cs
@@ -13,7 +13,7 @@ using System.Threading.Tasks;
 using System.Xml;
 using Tools.DotNETCommon;
 
-namespace BuildAgent.Listeners
+namespace BuildAgent.Run.Listeners
 {
 	/// <summary>
 	/// Writes parsed errors to a log file usable by Electric Commander
diff --git a/Engine/Source/Programs/BuildAgent/Run/Listeners/IssuesListener.cs b/Engine/Source/Programs/BuildAgent/Run/Listeners/IssuesListener.cs
new file mode 100644
--- /dev/null
+++ b/Engine/Source/Programs/BuildAgent/Run/Listeners/IssuesListener.cs
@@ -0,0 +1,117 @@
+using BuildAgent.Issues;
+using BuildAgent.Run.Interfaces;
+using System;
+using System.Collections.Generic;
+using System.IO;
+using System.Linq;
+using System.Runtime.Serialization.Json;
+using System.Text;
+using System.Threading.Tasks;
+using Tools.DotNETCommon;
+
+namespace BuildAgent.Run.Listeners
+{
+	/// <summary>
+	/// Issue listener. Generates an output report that can be consumed by the 
+	/// </summary>
+	class IssuesListener : IErrorListener
+	{
+		InputJob Job;
+		InputJobStep JobStep;
+		string LineUrl;
+		FileReference OutputFile;
+
+		public IssuesListener(string Stream, int Change, string JobName, string JobUrl, string JobStepName, string JobStepUrl, string LineUrl, string BaseDir, FileReference OutputFile)
+		{
+			this.OutputFile = OutputFile;
+			this.LineUrl = LineUrl;
+
+			JobStep = new InputJobStep();
+			JobStep.Name = JobStepName;
+			JobStep.Url = JobStepUrl;
+			JobStep.BaseDirectory = BaseDir;
+
+			Job = new InputJob();
+			Job.Change = Change;
+			Job.Stream = Stream;
+			Job.Url = JobUrl;
+			Job.Name = JobName;
+			Job.Steps.Add(JobStep);
+		}
+
+		public void Dispose()
+		{
+			// Try to write the output file in a transactional way; write it to a temporary file and rename it.
+			DirectoryReference.CreateDirectory(OutputFile.Directory);
+
+			InputData Data = new InputData();
+			Data.Jobs.Add(Job);
+
+			FileReference TempOutputFile = new FileReference(OutputFile.FullName + ".incoming");
+			using (MemoryStream Stream = new MemoryStream())
+			{
+				DataContractJsonSerializer InputFileDataSerializer = new DataContractJsonSerializer(typeof(InputData));
+				InputFileDataSerializer.WriteObject(Stream, Data);
+				FileReference.WriteAllBytes(TempOutputFile, Stream.ToArray());
+			}
+
+			FileReference.Delete(OutputFile);
+			FileReference.Move(TempOutputFile, OutputFile);
+		}
+
+		public void OnErrorMatch(ErrorMatch Error)
+		{
+			if (Error.Severity == ErrorSeverity.Error)
+			{
+				// Find the longest sequence of spaces common to every non-empty line
+				int WhitespaceLen = int.MaxValue;
+				foreach (string Line in Error.Lines)
+				{
+					int ThisWhitespaceLen = 0;
+					while (ThisWhitespaceLen < Line.Length && Line[ThisWhitespaceLen] == ' ')
+					{
+						ThisWhitespaceLen++;
+					}
+					if (ThisWhitespaceLen < Line.Length)
+					{
+						WhitespaceLen = Math.Min(WhitespaceLen, ThisWhitespaceLen);
+					}
+				}
+
+				// Remove the whitespace prefix
+				List<string> Lines = Error.Lines;
+				if (WhitespaceLen < int.MaxValue)
+				{
+					Lines = new List<string>(Lines);
+					for (int Idx = 0; Idx < Lines.Count; Idx++)
+					{
+						string Line = Lines[Idx];
+						if (Line.Length > WhitespaceLen)
+						{
+							Line = Line.Substring(WhitespaceLen);
+						}
+						else
+						{
+							Line = String.Empty;
+						}
+						Lines[Idx] = Line;
+					}
+				}
+
+				// Add the output diagnostic
+				InputDiagnostic Diagnostic = new InputDiagnostic();
+				Diagnostic.Type = "Error";
+				Diagnostic.Message = String.Join("\n", Lines);
+				if (LineUrl != null)
+				{
+					string Url = LineUrl;
+					Url = Url.Replace("{LINE_START}", Error.MinLineNumber.ToString());
+					Url = Url.Replace("{LINE_END}", Error.MaxLineNumber.ToString());
+					Url = Url.Replace("{LINE_COUNT}", (Error.MaxLineNumber + 1 - Error.MinLineNumber).ToString());
+					Diagnostic.Url = Url;
+				}
+				JobStep.Diagnostics.Add(Diagnostic);
+			}
+		}
+	}
+}
diff --git a/Engine/Source/Programs/BuildAgent/Run/RunMode.cs b/Engine/Source/Programs/BuildAgent/Run/RunMode.cs
--- a/Engine/Source/Programs/BuildAgent/Run/RunMode.cs
+++ b/Engine/Source/Programs/BuildAgent/Run/RunMode.cs
@@ -1,7 +1,7 @@
 // Copyright 1998-2019 Epic Games, Inc. All Rights Reserved.
 
-using BuildAgent.Listeners;
 using BuildAgent.Run.Interfaces;
+using BuildAgent.Run.Listeners;
 using System;
 using System.Collections.Generic;
 using System.ComponentModel;
@@ -53,6 +53,42 @@ namespace BuildAgent.Run
 		[Description("Enables the ElectricCommander listener")]
 		bool bElectricCommanderListener = false;
 
+		[CommandLine("-Stream=")]
+		[Description("Specifies the current stream (for issues output)")]
+		string Stream;
+
+		[CommandLine("-Change=")]
+		[Description("Specifies the current CL (for issues output)")]
+		int Change;
+
+		[CommandLine("-JobName=")]
+		[Description("Specifies the current job name (for issues output)")]
+		string JobName;
+
+		[CommandLine("-JobUrl=")]
+		[Description("Specifies the current job url (for issues output)")]
+		string JobUrl;
+
+		[CommandLine("-JobStepName=")]
+		[Description("Specifies the current job step name (for issues output)")]
+		string JobStepName;
+
+		[CommandLine("-JobStepUrl=")]
+		[Description("Specifies the current job step url (for issues output)")]
+		string JobStepUrl;
+
+		[CommandLine("-LineUrl=")]
+		[Description("Specifies a template for the url to a specific output line (for issues output)")]
+		string LineUrl;
+
+		[CommandLine("-BaseDir=")]
+		[Description("Specifies the base directory (for issues output)")]
+		string BaseDir;
+
+		[CommandLine("-IssuesOutput=")]
+		[Description("Specifies an output file for build issues")]
+		FileReference IssuesOutputFile = null;
+
 		public override void Configure(CommandLineArguments Arguments)
 		{
 			base.Configure(Arguments);
@@ -67,6 +103,34 @@ namespace BuildAgent.Run
 				bElectricCommanderListener = true;
 			}
 
+			if (IssuesOutputFile != null)
+			{
+				if (Stream == null)
+				{
+					throw new CommandLineArgumentException("Missing -Stream=... argument when specifying -IssuesOutput=...");
+				}
+				if (Change == 0)
+				{
+					throw new CommandLineArgumentException("Missing -Change=... argument when specifying -IssuesOutput=...");
+				}
+				if (JobName == null)
+				{
+					throw new CommandLineArgumentException("Missing -JobName=... argument when specifying -IssuesOutput=...");
+				}
+				if (JobUrl == null)
+				{
+					throw new CommandLineArgumentException("Missing -JobUrl=... argument when specifying -IssuesOutput=...");
+				}
+				if (JobStepName == null)
+				{
+					throw new CommandLineArgumentException("Missing -JobStepName=... argument when specifying -IssuesOutput=...");
+				}
+				if (JobStepUrl == null)
+				{
+					throw new CommandLineArgumentException("Missing -JobStepUrl=... argument when specifying -IssuesOutput=...");
+				}
+			}
+
 			ProgramArguments = Arguments.GetPositionalArguments();
 		}
 
@@ -125,6 +189,10 @@ namespace BuildAgent.Run
 				{
 					Listeners.Add(new ElectricCommanderListener());
 				}
+				if (IssuesOutputFile != null)
+				{
+					Listeners.Add(new IssuesListener(Stream, Change, JobName, JobUrl, JobStepName, JobStepUrl, LineUrl, BaseDir, IssuesOutputFile));
+				}
 
 				// Process the input
 				if (InputFile != null)
