BuildAgent: Avoid renaming EC diagnostics file; if something has a handle to the file open, this will fail. Instead, just open the existing file and write new data over the top.

#rnx

diff --git a/Engine/Source/Programs/BuildAgent/Run/Listeners/ElectricCommanderListener.cs b/Engine/Source/Programs/BuildAgent/Run/Listeners/ElectricCommanderListener.cs
--- a/Engine/Source/Programs/BuildAgent/Run/Listeners/ElectricCommanderListener.cs
+++ b/Engine/Source/Programs/BuildAgent/Run/Listeners/ElectricCommanderListener.cs
@@ -151,9 +151,7 @@ namespace BuildAgent.Run.Listeners
 				}
 
 				// Write them to disk
-				FileReference TempOutputFile = new FileReference(OutputFile.FullName + ".new");
-				Write(TempOutputFile, Errors);
-				FileUtils.ForceMoveFile(TempOutputFile, OutputFile);
+				Write(OutputFile, Errors);
 
 				// On the first run, set the path to the diagnostics file
 				if (!bHasSetDiagFile)
@@ -217,11 +215,14 @@ namespace BuildAgent.Run.Listeners
 		/// </summary>
 		static void Write(FileReference OutputFile, List<ErrorMatch> Errors)
 		{
+			byte[] XmlData;
+
+			// Prepare the XML data that needs to be written
 			XmlWriterSettings Settings = new XmlWriterSettings();
 			Settings.NewLineChars = Environment.NewLine;
 			Settings.Indent = true;
 			Settings.IndentChars = "\t";
-			using (FileStream Stream = File.Open(OutputFile.FullName, FileMode.Create, FileAccess.Write, FileShare.Read | FileShare.Delete))
+			using (MemoryStream Stream = new MemoryStream())
 			{
 				using (XmlWriter Writer = XmlWriter.Create(Stream, Settings))
 				{
@@ -235,7 +236,7 @@ namespace BuildAgent.Run.Listeners
 							Writer.WriteElementString("matcher", Error.Type);
 							Writer.WriteElementString("name", "");
 
-							if(Error.Severity == ErrorSeverity.Error)
+							if (Error.Severity == ErrorSeverity.Error)
 							{
 								Writer.WriteElementString("type", "error");
 							}
@@ -248,7 +249,7 @@ namespace BuildAgent.Run.Listeners
 							Writer.WriteElementString("numLines", (Error.MaxLineNumber + 1 - Error.MinLineNumber).ToString());
 
 							StringBuilder Message = new StringBuilder();
-							foreach(string Line in Error.Lines)
+							foreach (string Line in Error.Lines)
 							{
 								Message.AppendLine(Line);
 							}
@@ -259,6 +260,15 @@ namespace BuildAgent.Run.Listeners
 					}
 					Writer.WriteEndElement();
 				}
+				XmlData = Stream.ToArray();
+			}
+
+			// Open the existing file, write the new data, and truncate it.
+			using (FileStream Stream = FileReference.Open(OutputFile, FileMode.OpenOrCreate, FileAccess.Write, FileShare.Read | FileShare.Delete))
+			{
+				Stream.Position = 0;
+				Stream.Write(XmlData, 0, XmlData.Length);
+				Stream.SetLength(Stream.Position);
 			}
 		}
 	}
