Add missing copyright
#rnx

diff --git a/Engine/Source/Programs/BuildAgent/Api/ApiMode.cs b/Engine/Source/Programs/BuildAgent/Api/ApiMode.cs
--- a/Engine/Source/Programs/BuildAgent/Api/ApiMode.cs
+++ b/Engine/Source/Programs/BuildAgent/Api/ApiMode.cs
@@ -1,3 +1,4 @@
+// Copyright 1998-2019 Epic Games, Inc. All Rights Reserved.
 using System;
 using System.Collections.Generic;
 using System.ComponentModel;
