Add copyright notices
#rnx

diff --git a/Engine/Source/Programs/BuildAgent/Issues/IssueDiagnostic.cs b/Engine/Source/Programs/BuildAgent/Issues/IssueDiagnostic.cs
--- a/Engine/Source/Programs/BuildAgent/Issues/IssueDiagnostic.cs
+++ b/Engine/Source/Programs/BuildAgent/Issues/IssueDiagnostic.cs
@@ -1,3 +1,4 @@
+// Copyright 1998-2019 Epic Games, Inc. All Rights Reserved.
 using System;
 using System.Collections.Generic;
 using System.Diagnostics;
diff --git a/Engine/Source/Programs/BuildAgent/Run/Listeners/IssuesListener.cs b/Engine/Source/Programs/BuildAgent/Run/Listeners/IssuesListener.cs
--- a/Engine/Source/Programs/BuildAgent/Run/Listeners/IssuesListener.cs
+++ b/Engine/Source/Programs/BuildAgent/Run/Listeners/IssuesListener.cs
@@ -1,3 +1,4 @@
+// Copyright 1998-2019 Epic Games, Inc. All Rights Reserved.
 using BuildAgent.Issues;
 using BuildAgent.Run.Interfaces;
 using System;
