BuildAgent: Remove some debug output, and add matcher for additional lines of ld output in linker errors.

#rnx

diff --git a/Engine/Source/Programs/BuildAgent/Run/Matchers/LinkErrorMatcher.cs b/Engine/Source/Programs/BuildAgent/Run/Matchers/LinkErrorMatcher.cs
--- a/Engine/Source/Programs/BuildAgent/Run/Matchers/LinkErrorMatcher.cs
+++ b/Engine/Source/Programs/BuildAgent/Run/Matchers/LinkErrorMatcher.cs
@@ -17,7 +17,7 @@ namespace BuildAgent.Matchers
 		{
 			if(Input.IsMatch(@"error: linker command failed with exit code "))
 			{
-				int MinIdx = Input.MatchBackwards(0, ": In function |: undefined reference to ");
+				int MinIdx = Input.MatchBackwards(0, ": In function |: undefined reference to |[^a-zA-Z]ld: ");
 				return new ErrorMatch(ErrorSeverity.Error, ErrorPriority.Normal, "Link", Input, MinIdx, 0);
 			}
 			if (Input.IsMatch("Undefined symbols for architecture"))
diff --git a/Engine/Source/Programs/BuildAgent/Run/RunMode.cs b/Engine/Source/Programs/BuildAgent/Run/RunMode.cs
--- a/Engine/Source/Programs/BuildAgent/Run/RunMode.cs
+++ b/Engine/Source/Programs/BuildAgent/Run/RunMode.cs
@@ -269,14 +269,8 @@ namespace BuildAgent.Run
 			LineBuffer Buffer = new LineBuffer(ReadLine, 50);
 			ReadOnlyLineBuffer ReadOnlyBuffer = new ReadOnlyLineBuffer(Buffer);
 
-			Stopwatch Timer = Stopwatch.StartNew();
 			while (Buffer[0] != null)
 			{
-				if ((Buffer.CurrentLineNumber % 1000) == 0)
-				{
-					Log.TraceInformation("LINE: {0} ({1}s)", Buffer.CurrentLineNumber, (int)Timer.Elapsed.TotalSeconds);
-				}
-
 				// Try to match an error
 				ErrorMatch Error = null;
 				foreach (IErrorMatcher Matcher in Matchers)
