Making BuildAgent.exe not prefer 32-bit so it doesn't crash.

#rnx

diff --git a/Engine/Source/Programs/BuildAgent/BuildAgent.csproj b/Engine/Source/Programs/BuildAgent/BuildAgent.csproj
--- a/Engine/Source/Programs/BuildAgent/BuildAgent.csproj
+++ b/Engine/Source/Programs/BuildAgent/BuildAgent.csproj
@@ -33,6 +33,7 @@
     <ErrorReport>prompt</ErrorReport>
     <WarningLevel>4</WarningLevel>
     <NoWarn>CS0649</NoWarn>
+    <Prefer32Bit>false</Prefer32Bit>
   </PropertyGroup>
   <Import Project="NotForLicensees\EpicInternal.props" Condition="Exists('NotForLicensees\EpicInternal.props')" />
   <ItemGroup>
