Remove invalid characters serialized to error output. Some characters are not serializable to XML even if they are escaped.

#rnx

diff --git a/Engine/Source/Programs/BuildAgent/Run/Listeners/ElectricCommanderListener.cs b/Engine/Source/Programs/BuildAgent/Run/Listeners/ElectricCommanderListener.cs
--- a/Engine/Source/Programs/BuildAgent/Run/Listeners/ElectricCommanderListener.cs
+++ b/Engine/Source/Programs/BuildAgent/Run/Listeners/ElectricCommanderListener.cs
@@ -253,7 +253,7 @@ namespace BuildAgent.Run.Listeners
 							{
 								Message.AppendLine(Line);
 							}
-							Writer.WriteElementString("message", Message.ToString());
+							 Writer.WriteElementString("message", SanitizeString(Message.ToString()));
 
 							Writer.WriteEndElement();
 						}
@@ -271,5 +271,25 @@ namespace BuildAgent.Run.Listeners
 				Stream.SetLength(Stream.Position);
 			}
 		}
+
+		/// <summary>
+		/// Removes invalid characters from the string. Some characters are invalid even when encoded into entities.
+		/// </summary>
+		/// <param name="Line">Line to sanitize</param>
+		/// <returns></returns>
+		static string SanitizeString(string Line)
+		{
+			// Don't expect many (if any) hits here, so just operate on whole string objects
+			for (int Idx = 0; Idx < Line.Length; Idx++)
+			{
+				int Character = Line[Idx];
+				if(Character < 0x20 && (Character != 0x09 && Character != 0x0a && Character != 0x0d))
+				{
+					Line = Line.Remove(Idx, 1);
+					Idx--;
+				}
+			}
+			return Line;
+		}
 	}
 }
