Fix race condition in BuildAgent when flushing data to EC. bDisposing can be read outside the lock, and causes the background thread to exit before the buffer has drained. Capture it inside the lock instead.

#rnx

diff --git a/Engine/Source/Programs/BuildAgent/Run/Listeners/ElectricCommanderListener.cs b/Engine/Source/Programs/BuildAgent/Run/Listeners/ElectricCommanderListener.cs
--- a/Engine/Source/Programs/BuildAgent/Run/Listeners/ElectricCommanderListener.cs
+++ b/Engine/Source/Programs/BuildAgent/Run/Listeners/ElectricCommanderListener.cs
@@ -134,20 +134,18 @@ namespace BuildAgent.Run.Listeners
 			string Outcome = null;
 
 			List<ErrorMatch> Errors = new List<ErrorMatch>();
-			while (!bDisposing)
+			for(; ;)
 			{
 				// Copy the current set of errors
-				lock(LockObject)
+				bool bReadyToDispose;
+				lock (LockObject)
 				{
 					if (NewErrors.Count > 0)
 					{
 						Errors.AddRange(NewErrors);
 						NewErrors.Clear();
 					}
-					else if(bDisposing)
-					{
-						break;
-					}
+					bReadyToDispose = bDisposing;
 				}
 
 				// Write them to disk
@@ -184,6 +182,12 @@ namespace BuildAgent.Run.Listeners
 					Outcome = NewOutcome;
 				}
 
+				// Exit once we've flushed everything
+				if(bReadyToDispose)
+				{
+					break;
+				}
+
 				// Wait until the next update
 				UpdateEvent.WaitOne();
 			}
