Removing a debug output I added to build agent.

#rnx

diff --git a/Engine/Source/Programs/BuildAgent/Program.cs b/Engine/Source/Programs/BuildAgent/Program.cs
--- a/Engine/Source/Programs/BuildAgent/Program.cs
+++ b/Engine/Source/Programs/BuildAgent/Program.cs
@@ -25,7 +25,6 @@ namespace BuildAgent
 		{
 			try
 			{
-				Log.TraceInformation("Welcome to the Danger Zone!");
 				int Result = GuardedMain(Args);
 				return Result;
 			}
