Remove some debug output from BuildAgent.

#rnx

diff --git a/Engine/Source/Programs/BuildAgent/Run/Matchers/CompileErrorMatcher.cs b/Engine/Source/Programs/BuildAgent/Run/Matchers/CompileErrorMatcher.cs
--- a/Engine/Source/Programs/BuildAgent/Run/Matchers/CompileErrorMatcher.cs
+++ b/Engine/Source/Programs/BuildAgent/Run/Matchers/CompileErrorMatcher.cs
@@ -8,7 +8,6 @@ using System.Linq;
 using System.Text;
 using System.Text.RegularExpressions;
 using System.Threading.Tasks;
-using Tools.DotNETCommon;
 
 namespace BuildAgent.Matchers
 {
@@ -65,7 +64,6 @@ namespace BuildAgent.Matchers
 				EndIdx++;
 			}
 
-			Log.TraceWarning("WARNING DETECTED! {0}", Input);
 			return new ErrorMatch(Severity, ErrorPriority.High, "Compile", Input, 0, EndIdx);
 		}
 
