Change location of EpicInternal.props.

diff --git a/Engine/Source/Programs/BuildAgent/BuildAgent.csproj b/Engine/Source/Programs/BuildAgent/BuildAgent.csproj
--- a/Engine/Source/Programs/BuildAgent/BuildAgent.csproj
+++ b/Engine/Source/Programs/BuildAgent/BuildAgent.csproj
@@ -36,7 +36,7 @@
     <NoWarn>CS0649</NoWarn>
     <Prefer32Bit>false</Prefer32Bit>
   </PropertyGroup>
-  <Import Project="NotForLicensees\EpicInternal.props" Condition="Exists('NotForLicensees\EpicInternal.props')" />
+  <Import Project="..\..\..\Restricted\NotForLicensees\Source\Programs\BuildAgent\EpicInternal.props" Condition="Exists('..\..\..\Restricted\NotForLicensees\Source\Programs\BuildAgent\EpicInternal.props')" />
   <ItemGroup>
     <Reference Include="System" />
     <Reference Include="System.Core" />
