Separate restricted files for BuildAgent into separate folders.

diff --git a/Engine/Source/Programs/BuildAgent/BuildAgent.csproj b/Engine/Source/Programs/BuildAgent/BuildAgent.csproj
--- a/Engine/Source/Programs/BuildAgent/BuildAgent.csproj
+++ b/Engine/Source/Programs/BuildAgent/BuildAgent.csproj
@@ -81,6 +81,12 @@
     <Compile Include="Run\Listeners\DebugOutputListener.cs" />
     <Compile Include="Run\Listeners\IssuesListener.cs" />
     <Compile Include="Run\Matchers\**\*.cs" />
+    <Compile Include="..\..\..\Platforms\*\Source\Programs\BuildAgent\**\*.cs">
+      <Link>$([System.Text.RegularExpressions.Regex]::Replace(%(Compile.RecursiveDir), ^([^\\/]+[\\/]){4}, ``))/_</Link>
+    </Compile>
+    <Compile Include="..\..\..\Restricted\*\Source\Programs\BuildAgent\**\*.cs">
+      <Link>$([System.Text.RegularExpressions.Regex]::Replace(%(Compile.RecursiveDir), ^([^\\/]+[\\/]){4}, ``))/_</Link>
+    </Compile>
     <Compile Include="Program.cs" />
     <Compile Include="Run\ErrorMatch.cs" />
     <Compile Include="Run\Interfaces\IErrorListener.cs" />
@@ -112,7 +118,7 @@
     <None Include="App.config" />
   </ItemGroup>
   <ItemGroup>
-    <ProjectReference Include="..\..\..\Source\Programs\DotNETCommon\DotNETUtilities\DotNETUtilities.csproj">
+    <ProjectReference Include="..\DotNETCommon\DotNETUtilities\DotNETUtilities.csproj">
       <Project>{5d7d66e8-8c76-4af9-b3ec-2ef03421d730}</Project>
       <Name>DotNETUtilities</Name>
     </ProjectReference>
